Fix non-unity compile error from Dev-VR.

diff --git a/Engine/Plugins/Runtime/Steam/SteamAudio/Source/SteamAudio/Private/PhononProbeVolume.cpp b/Engine/Plugins/Runtime/Steam/SteamAudio/Source/SteamAudio/Private/PhononProbeVolume.cpp
--- a/Engine/Plugins/Runtime/Steam/SteamAudio/Source/SteamAudio/Private/PhononProbeVolume.cpp
+++ b/Engine/Plugins/Runtime/Steam/SteamAudio/Source/SteamAudio/Private/PhononProbeVolume.cpp
@@ -8,6 +8,8 @@
 #include "PhononScene.h"
 #include "PhononCommon.h"
 
+#include "Components/PrimitiveComponent.h"
+
 #if WITH_EDITOR
 #include "Editor.h"
 #include "LevelEditorViewport.h"
