Copying //UE4/Dev-AnimPhys to //UE4/Dev-Main (Source: //UE4/Dev-AnimPhys @ 3624379)

diff --git a/Engine/Plugins/Runtime/Steam/SteamAudio/Source/SteamAudio/Private/PhononPluginManager.cpp b/Engine/Plugins/Runtime/Steam/SteamAudio/Source/SteamAudio/Private/PhononPluginManager.cpp
new file mode 100644
--- /dev/null
+++ b/Engine/Plugins/Runtime/Steam/SteamAudio/Source/SteamAudio/Private/PhononPluginManager.cpp
@@ -0,0 +1,98 @@
+//
+// Copyright (C) Valve Corporation. All rights reserved.
+//
+
+#include "PhononPluginManager.h"
+#include "SteamAudioModule.h"
+#include "AudioDevice.h"
+
+namespace SteamAudio
+{
+	FPhononPluginManager::FPhononPluginManager()
+		: bEnvironmentCreated(false)
+		, ReverbPtr(nullptr)
+		, OcclusionPtr(nullptr)
+	{
+	}
+
+	FPhononPluginManager::~FPhononPluginManager()
+	{
+		if (bEnvironmentCreated)
+		{
+			Environment.Shutdown();
+			bEnvironmentCreated = false;
+		}
+	}
+
+	void FPhononPluginManager::OnListenerInitialize(FAudioDevice* AudioDevice, UWorld* ListenerWorld)
+	{
+		IPLhandle RendererPtr = Environment.Initialize(ListenerWorld, AudioDevice);
+
+		// If we've succeeded, pass the phonon environmental renderer to the occlusion and reverb plugins, if we're using them
+		if (RendererPtr)
+		{
+			if (IsUsingSteamAudioPlugin(EAudioPlugin::REVERB))
+			{
+				ReverbPtr = (FPhononReverb*)AudioDevice->ReverbPluginInterface.Get();
+				ReverbPtr->SetEnvironmentalRenderer(RendererPtr);
+				ReverbPtr->SetEnvironmentCriticalSection(Environment.GetEnvironmentCriticalSection());
+				ReverbPtr->CreateReverbEffect();
+			}
+
+			if (IsUsingSteamAudioPlugin(EAudioPlugin::OCCLUSION))
+			{
+				OcclusionPtr = (FPhononOcclusion*)AudioDevice->OcclusionInterface.Get();
+				OcclusionPtr->SetEnvironmentalRenderer(RendererPtr);
+				OcclusionPtr->SetCriticalSectionHandle(Environment.GetEnvironmentCriticalSection());
+			}
+
+			bEnvironmentCreated = true;
+		}
+	}
+
+	void FPhononPluginManager::OnListenerUpdated(FAudioDevice* AudioDevice, const int32 ViewportIndex, const FTransform& ListenerTransform, const float InDeltaSeconds)
+	{
+		if (!bEnvironmentCreated)
+		{
+			return;
+		}
+
+		FVector Position = ListenerTransform.GetLocation();
+		FVector Forward = ListenerTransform.GetUnitAxis(EAxis::Y);
+		FVector Up = ListenerTransform.GetUnitAxis(EAxis::Z);
+
+		if (OcclusionPtr)
+		{
+			OcclusionPtr->UpdateDirectSoundSources(Position, Forward, Up);
+		}
+
+		if (ReverbPtr)
+		{
+			ReverbPtr->UpdateListener(Position, Forward, Up);
+		}
+	}
+
+	void FPhononPluginManager::OnListenerShutdown(FAudioDevice* AudioDevice)
+	{
+		FSteamAudioModule* Module = &FModuleManager::GetModuleChecked<FSteamAudioModule>("SteamAudio");
+		if (Module != nullptr)
+		{
+			Module->UnregisterAudioDevice(AudioDevice);
+		}
+	}
+
+	bool FPhononPluginManager::IsUsingSteamAudioPlugin(EAudioPlugin PluginType)
+	{
+		FSteamAudioModule* Module = &FModuleManager::GetModuleChecked<FSteamAudioModule>("SteamAudio");
+
+		// If we can't get the module from the module manager, then we don't have any of these plugins loaded.
+		if (Module == nullptr)
+		{
+			return false;
+		}
+
+		FString SteamPluginName = Module->GetPluginFactory(PluginType)->GetDisplayName();
+		FString CurrentPluginName = AudioPluginUtilities::GetDesiredPluginName(PluginType, AudioPluginUtilities::CurrentPlatform);
+		return CurrentPluginName.Equals(SteamPluginName);
+	}
+}
diff --git a/Engine/Plugins/Runtime/Steam/SteamAudio/Source/SteamAudio/Private/PhononPluginManager.h b/Engine/Plugins/Runtime/Steam/SteamAudio/Source/SteamAudio/Private/PhononPluginManager.h
new file mode 100644
--- /dev/null
+++ b/Engine/Plugins/Runtime/Steam/SteamAudio/Source/SteamAudio/Private/PhononPluginManager.h
@@ -0,0 +1,44 @@
+//
+// Copyright (C) Valve Corporation. All rights reserved.
+//
+
+#pragma once
+
+#include "SteamAudioModule.h"
+#include "PhononSpatialization.h"
+#include "PhononOcclusion.h"
+#include "PhononReverb.h"
+#include "SteamAudioEnvironment.h"
+#include "AudioPluginUtilities.h"
+
+namespace SteamAudio
+{
+	/************************************************************************/
+	/* FPhononPluginManager                                                 */
+	/* This ListenerObserver owns the Steam Audio environment, and          */
+	/* dispatches information to the Steam Audio reverb and occlusion       */
+	/* plugins.                                                             */
+	/************************************************************************/
+	class FPhononPluginManager : public IAudioPluginListener
+	{
+	public:
+		FPhononPluginManager();
+		~FPhononPluginManager();
+
+		//~ Begin IAudioPluginListener
+		virtual void OnListenerInitialize(FAudioDevice* AudioDevice, UWorld* ListenerWorld) override;
+		virtual void OnListenerUpdated(FAudioDevice* AudioDevice, const int32 ViewportIndex, const FTransform& ListenerTransform, const float InDeltaSeconds) override;
+		virtual void OnListenerShutdown(FAudioDevice* AudioDevice) override;
+		//~ End IAudioPluginListener
+
+	private:
+		bool bEnvironmentCreated;
+		FEnvironment Environment;
+		FPhononReverb* ReverbPtr;
+		FPhononOcclusion* OcclusionPtr;
+
+	private:
+		/* Helper function for checking whether the user is using Steam Audio for spatialization, reverb, and/or occlusion: */
+		static bool IsUsingSteamAudioPlugin(EAudioPlugin PluginType);
+	};
+}
diff --git a/Engine/Plugins/Runtime/Steam/SteamAudio/Source/SteamAudio/Private/PhononReverb.cpp b/Engine/Plugins/Runtime/Steam/SteamAudio/Source/SteamAudio/Private/PhononReverb.cpp
--- a/Engine/Plugins/Runtime/Steam/SteamAudio/Source/SteamAudio/Private/PhononReverb.cpp
+++ b/Engine/Plugins/Runtime/Steam/SteamAudio/Source/SteamAudio/Private/PhononReverb.cpp
@@ -10,6 +10,9 @@
 #include "Sound/SoundSubmix.h"
 #include "DSP/Dsp.h"
 #include "ScopeLock.h"
+#include "PhononPluginManager.h"
+
+#include "AudioDevice.h"
 
 namespace SteamAudio
 {
@@ -25,7 +28,8 @@ namespace SteamAudio
 		, ReverbConvolutionEffect(nullptr)
 		, AmbisonicsChannels(0)
 		, IndirectOutDeinterleaved(nullptr)
-		, SteamAudioModule(nullptr)
+		, CachedSpatializationMethod(EIplSpatializationMethod::PANNING)
+		, EnvironmentalCriticalSectionHandle(nullptr)
 	{
 		const int32 IndirectImpulseResponseOrder = GetDefault<USteamAudioSettings>()->IndirectImpulseResponseOrder;
 
@@ -108,13 +112,11 @@ namespace SteamAudio
 		}
 	}
 
-	void FPhononReverb::Initialize(const int32 SampleRate, const int32 NumSources, const int32 FrameSize)
+	void FPhononReverb::Initialize(const FAudioPluginInitializationParams InitializationParams)
 	{
-		SteamAudioModule = &FModuleManager::GetModuleChecked<FSteamAudioModule>("SteamAudio");
-
 		RenderingSettings.convolutionType = IPL_CONVOLUTIONTYPE_PHONON;
-		RenderingSettings.frameSize = FrameSize;
-		RenderingSettings.samplingRate = SampleRate;
+		RenderingSettings.frameSize = InitializationParams.BufferLength;
+		RenderingSettings.samplingRate = InitializationParams.SampleRate;
 
 		IPLHrtfParams HrtfParams;
 		HrtfParams.hrtfData = nullptr;
@@ -134,32 +136,34 @@ namespace SteamAudio
 
 		for (int32 i = 0; i < AmbisonicsChannels; ++i)
 		{
-			IndirectOutDeinterleaved[i] = new float[FrameSize];
+			IndirectOutDeinterleaved[i] = new float[InitializationParams.BufferLength];
 		}
 		IndirectIntermediateBuffer.format = IndirectOutputAudioFormat;
-		IndirectIntermediateBuffer.numSamples = FrameSize;
+		IndirectIntermediateBuffer.numSamples = InitializationParams.BufferLength;
 		IndirectIntermediateBuffer.interleavedBuffer = nullptr;
 		IndirectIntermediateBuffer.deinterleavedBuffer = IndirectOutDeinterleaved;
 
 		DryBuffer.format = ReverbInputAudioFormat;
-		DryBuffer.numSamples = FrameSize;
+		DryBuffer.numSamples = InitializationParams.BufferLength;
 		DryBuffer.interleavedBuffer = nullptr;
 		DryBuffer.deinterleavedBuffer = nullptr;
 
-		IndirectOutArray.SetNumZeroed(FrameSize * BinauralOutputAudioFormat.numSpeakers);
+		IndirectOutArray.SetNumZeroed(InitializationParams.BufferLength * BinauralOutputAudioFormat.numSpeakers);
 		IndirectOutBuffer.format = BinauralOutputAudioFormat;
-		IndirectOutBuffer.numSamples = FrameSize;
+		IndirectOutBuffer.numSamples = InitializationParams.BufferLength;
 		IndirectOutBuffer.interleavedBuffer = IndirectOutArray.GetData();
 		IndirectOutBuffer.deinterleavedBuffer = nullptr;
 
-		ReverbSources.SetNum(NumSources);
-		for (auto& ReverbSource : ReverbSources)
+		ReverbSources.SetNum(InitializationParams.NumSources);
+		for (FReverbSource& ReverbSource : ReverbSources)
 		{
 			ReverbSource.InBuffer.format = InputAudioFormat;
-			ReverbSource.InBuffer.numSamples = FrameSize;
+			ReverbSource.InBuffer.numSamples = InitializationParams.BufferLength;
 		}
 
 		ReverbIndirectContribution = 1.0f;
+
+		CachedSpatializationMethod = GetDefault<USteamAudioSettings>()->IndirectSpatializationMethod;
 	}
 
 	void FPhononReverb::OnInitSource(const uint32 SourceId, const FName& AudioComponentUserId, const uint32 NumChannels, UReverbPluginSourceSettingsBase* InSettings)
@@ -172,8 +176,8 @@ namespace SteamAudio
 
 		UE_LOG(LogSteamAudio, Log, TEXT("Creating reverb effect for %s"), *AudioComponentUserId.ToString().ToLower());
 
-		auto Settings = static_cast<UPhononReverbSourceSettings*>(InSettings);
-		auto& ReverbSource = ReverbSources[SourceId];
+		UPhononReverbSourceSettings* Settings = static_cast<UPhononReverbSourceSettings*>(InSettings);
+		FReverbSource& ReverbSource = ReverbSources[SourceId];
 
 		ReverbSource.IndirectContribution = Settings->IndirectContribution;
 
@@ -215,20 +219,20 @@ namespace SteamAudio
 
 	void FPhononReverb::ProcessSourceAudio(const FAudioPluginSourceInputData& InputData, FAudioPluginSourceOutputData& OutputData)
 	{
-		if (!EnvironmentalRenderer)
+		if (!EnvironmentalRenderer || !EnvironmentalCriticalSectionHandle)
 		{
 			return;
 		}
 
-		FScopeLock EnvironmentLock(&SteamAudioModule->GetEnvironmentCriticalSection());
+		FScopeLock EnvironmentLock(EnvironmentalCriticalSectionHandle);
 
-		auto& ReverbSource = ReverbSources[InputData.SourceId];
-		auto Position = SteamAudio::UnrealToPhononIPLVector3(InputData.SpatializationParams->EmitterWorldPosition);
+		FReverbSource& ReverbSource = ReverbSources[InputData.SourceId];
+		IPLVector3 Position = SteamAudio::UnrealToPhononIPLVector3(InputData.SpatializationParams->EmitterWorldPosition);
 
 		if (ReverbSource.ConvolutionEffect)
 		{
 			ReverbSource.IndirectInArray.SetNumUninitialized(InputData.AudioBuffer->Num());
-			for (auto i = 0; i < InputData.AudioBuffer->Num(); ++i)
+			for (int32 i = 0; i < InputData.AudioBuffer->Num(); ++i)
 			{
 				ReverbSource.IndirectInArray[i] = (*InputData.AudioBuffer)[i] * ReverbSource.IndirectContribution;
 			}
@@ -240,12 +244,12 @@ namespace SteamAudio
 
 	void FPhononReverb::ProcessMixedAudio(const FSoundEffectSubmixInputData& InData, FSoundEffectSubmixOutputData& OutData)
 	{
-		if (!EnvironmentalRenderer)
+		if (!EnvironmentalRenderer || !EnvironmentalCriticalSectionHandle)
 		{
 			return;
 		}
 
-		FScopeLock EnvironmentLock(&SteamAudioModule->GetEnvironmentCriticalSection());
+		FScopeLock EnvironmentLock(EnvironmentalCriticalSectionHandle);
 
 		if (IndirectOutBuffer.format.numSpeakers != OutData.NumChannels)
 		{
@@ -272,7 +276,7 @@ namespace SteamAudio
 		if (ReverbConvolutionEffect)
 		{
 			ReverbIndirectInArray.SetNumUninitialized(InData.AudioBuffer->Num());
-			for (auto i = 0; i < InData.AudioBuffer->Num(); ++i)
+			for (int32 i = 0; i < InData.AudioBuffer->Num(); ++i)
 			{
 				ReverbIndirectInArray[i] = (*InData.AudioBuffer)[i] * ReverbIndirectContribution;
 			}
@@ -283,7 +287,7 @@ namespace SteamAudio
 
 		iplGetMixedEnvironmentalAudio(EnvironmentalRenderer, ListenerPosition, ListenerForward, ListenerUp, IndirectIntermediateBuffer);
 
-		switch (GetDefault<USteamAudioSettings>()->IndirectSpatializationMethod)
+		switch (CachedSpatializationMethod)
 		{
 		case EIplSpatializationMethod::HRTF:
 			iplApplyAmbisonicsBinauralEffect(IndirectBinauralEffect, IndirectIntermediateBuffer, IndirectOutBuffer);
@@ -298,9 +302,9 @@ namespace SteamAudio
 
 	void FPhononReverb::CreateReverbEffect()
 	{
-		check(EnvironmentalRenderer);
+		check(EnvironmentalRenderer && EnvironmentalCriticalSectionHandle);
 
-		FScopeLock Lock(&SteamAudioModule->GetEnvironmentCriticalSection());
+		FScopeLock Lock(EnvironmentalCriticalSectionHandle);
 
 		ReverbIndirectContribution = GetDefault<USteamAudioSettings>()->IndirectContribution;
 		switch (GetDefault<USteamAudioSettings>()->ReverbSimulationType)
@@ -339,6 +343,11 @@ namespace SteamAudio
 		EnvironmentalRenderer = InEnvironmentalRenderer;
 	}
 
+	void FPhononReverb::SetEnvironmentCriticalSection(FCriticalSection* CriticalSection)
+	{
+		EnvironmentalCriticalSectionHandle = CriticalSection;
+	}
+
 	//==============================================================================================================================================
 	// FReverbSource
 	//==============================================================================================================================================
diff --git a/Engine/Plugins/Runtime/Steam/SteamAudio/Source/SteamAudio/Private/PhononReverb.h b/Engine/Plugins/Runtime/Steam/SteamAudio/Source/SteamAudio/Private/PhononReverb.h
--- a/Engine/Plugins/Runtime/Steam/SteamAudio/Source/SteamAudio/Private/PhononReverb.h
+++ b/Engine/Plugins/Runtime/Steam/SteamAudio/Source/SteamAudio/Private/PhononReverb.h
@@ -23,13 +23,16 @@ namespace SteamAudio
 		TArray<float> IndirectInArray;
 	};
 
+	/************************************************************************/
+	/* FPhononReverb. Reverb plugin for Steam Audio.                        */
+	/************************************************************************/
 	class FPhononReverb : public IAudioReverb
 	{
 	public:
 		FPhononReverb();
 		~FPhononReverb();
 
-		virtual void Initialize(const int32 SampleRate, const int32 NumSources, const int32 FrameSize) override;
+		virtual void Initialize(const FAudioPluginInitializationParams InitializationParams) override;
 		virtual void OnInitSource(const uint32 SourceId, const FName& AudioComponentUserId, const uint32 NumChannels, UReverbPluginSourceSettingsBase* InSettings) override;
 		virtual void OnReleaseSource(const uint32 SourceId) override;
 		virtual class FSoundEffectSubmix* GetEffectSubmix(class USoundSubmix* Submix) override;
@@ -37,6 +40,7 @@ namespace SteamAudio
 		
 		void ProcessMixedAudio(const FSoundEffectSubmixInputData& InData, FSoundEffectSubmixOutputData& OutData);
 		void SetEnvironmentalRenderer(IPLhandle EnvironmentalRenderer);
+		void SetEnvironmentCriticalSection(FCriticalSection* CriticalSection);
 		void CreateReverbEffect();
 		void UpdateListener(const FVector& Position, const FVector& Forward, const FVector& Up);
 
@@ -66,6 +70,8 @@ namespace SteamAudio
 		IPLVector3 ListenerForward;
 		IPLVector3 ListenerUp;
 
+		EIplSpatializationMethod CachedSpatializationMethod;
+
 		IPLRenderingSettings RenderingSettings;
 
 		TArray<FReverbSource> ReverbSources;
@@ -73,7 +79,9 @@ namespace SteamAudio
 		float ReverbIndirectContribution;
 		TArray<float> ReverbIndirectInArray;
 
-		class FSteamAudioModule* SteamAudioModule;
+		TAudioPluginListenerPtr PluginManagerPtr;
+
+		FCriticalSection* EnvironmentalCriticalSectionHandle;
 	};
 }
 
diff --git a/Engine/Plugins/Runtime/Steam/SteamAudio/Source/SteamAudio/Private/PhononSpatialization.cpp b/Engine/Plugins/Runtime/Steam/SteamAudio/Source/SteamAudio/Private/PhononSpatialization.cpp
--- a/Engine/Plugins/Runtime/Steam/SteamAudio/Source/SteamAudio/Private/PhononSpatialization.cpp
+++ b/Engine/Plugins/Runtime/Steam/SteamAudio/Source/SteamAudio/Private/PhononSpatialization.cpp
@@ -42,11 +42,11 @@ namespace SteamAudio
 		}
 	}
 
-	void FPhononSpatialization::Initialize(const uint32 SampleRate, const uint32 NumSources, const uint32 OutputBufferLength)
+	void FPhononSpatialization::Initialize(const FAudioPluginInitializationParams InitializationParams)
 	{
 		RenderingSettings.convolutionType = IPL_CONVOLUTIONTYPE_PHONON;
-		RenderingSettings.frameSize = OutputBufferLength;
-		RenderingSettings.samplingRate = SampleRate;
+		RenderingSettings.frameSize = InitializationParams.BufferLength;
+		RenderingSettings.samplingRate = InitializationParams.SampleRate;
 
 		IPLHrtfParams HrtfParams;
 		HrtfParams.hrtfData = nullptr;
@@ -58,17 +58,17 @@ namespace SteamAudio
 
 		iplCreateBinauralRenderer(SteamAudio::GlobalContext, RenderingSettings, HrtfParams, &BinauralRenderer);
 
-		BinauralSources.AddDefaulted(NumSources);
-		for (auto& BinauralSource : BinauralSources)
+		BinauralSources.AddDefaulted(InitializationParams.NumSources);
+		for (FBinauralSource& BinauralSource : BinauralSources)
 		{
 			BinauralSource.InBuffer.format = InputAudioFormat;
-			BinauralSource.InBuffer.numSamples = OutputBufferLength;
+			BinauralSource.InBuffer.numSamples = InitializationParams.BufferLength;
 			BinauralSource.InBuffer.interleavedBuffer = nullptr;
 			BinauralSource.InBuffer.deinterleavedBuffer = nullptr;
 
-			BinauralSource.OutArray.SetNumZeroed(OutputBufferLength * 2);
+			BinauralSource.OutArray.SetNumZeroed(InitializationParams.BufferLength * 2);
 			BinauralSource.OutBuffer.format = BinauralOutputAudioFormat;
-			BinauralSource.OutBuffer.numSamples = OutputBufferLength;
+			BinauralSource.OutBuffer.numSamples = InitializationParams.BufferLength;
 			BinauralSource.OutBuffer.interleavedBuffer = nullptr;
 			BinauralSource.OutBuffer.deinterleavedBuffer = nullptr;
 		}
@@ -77,8 +77,8 @@ namespace SteamAudio
 	void FPhononSpatialization::OnInitSource(const uint32 SourceId, const FName& AudioComponentUserId,
 		USpatializationPluginSourceSettingsBase* InSettings)
 	{
-		auto SpatializationSettings = static_cast<UPhononSpatializationSourceSettings*>(InSettings);
-		auto& BinauralSource = BinauralSources[SourceId];
+		UPhononSpatializationSourceSettings* SpatializationSettings = static_cast<UPhononSpatializationSourceSettings*>(InSettings);
+		FBinauralSource& BinauralSource = BinauralSources[SourceId];
 
 		UE_LOG(LogSteamAudio, Log, TEXT("Creating spatialization effect."));
 
@@ -107,7 +107,7 @@ namespace SteamAudio
 
 	void FPhononSpatialization::OnReleaseSource(const uint32 SourceId)
 	{
-		auto& BinauralSource = BinauralSources[SourceId];
+		FBinauralSource& BinauralSource = BinauralSources[SourceId];
 
 		UE_LOG(LogSteamAudio, Log, TEXT("Destroying spatialization effect."));
 
@@ -124,7 +124,7 @@ namespace SteamAudio
 
 	void FPhononSpatialization::ProcessAudio(const FAudioPluginSourceInputData& InputData, FAudioPluginSourceOutputData& OutputData)
 	{
-		auto& BinauralSource = BinauralSources[InputData.SourceId];
+		FBinauralSource& BinauralSource = BinauralSources[InputData.SourceId];
 		BinauralSource.InBuffer.interleavedBuffer = InputData.AudioBuffer->GetData();
 		BinauralSource.OutBuffer.interleavedBuffer = OutputData.AudioBuffer.GetData();
 
@@ -140,12 +140,10 @@ namespace SteamAudio
 		switch (BinauralSource.SpatializationMethod)
 		{
 		case EIplSpatializationMethod::HRTF:
-			iplApplyBinauralEffect(BinauralSource.BinauralEffect, BinauralSource.InBuffer, RelativeDirection,
-				static_cast<IPLHrtfInterpolation>(BinauralSource.HrtfInterpolationMethod), BinauralSource.OutBuffer);
+			iplApplyBinauralEffect(BinauralSource.BinauralEffect, BinauralSource.InBuffer, RelativeDirection, static_cast<IPLHrtfInterpolation>(BinauralSource.HrtfInterpolationMethod), BinauralSource.OutBuffer);
 			break;
 		case EIplSpatializationMethod::PANNING:
-			iplApplyPanningEffect(BinauralSource.PanningEffect, BinauralSource.InBuffer, RelativeDirection,
-				BinauralSource.OutBuffer);
+			iplApplyPanningEffect(BinauralSource.PanningEffect, BinauralSource.InBuffer, RelativeDirection, BinauralSource.OutBuffer);
 			break;
 		}
 	}
diff --git a/Engine/Plugins/Runtime/Steam/SteamAudio/Source/SteamAudio/Private/PhononSpatialization.h b/Engine/Plugins/Runtime/Steam/SteamAudio/Source/SteamAudio/Private/PhononSpatialization.h
--- a/Engine/Plugins/Runtime/Steam/SteamAudio/Source/SteamAudio/Private/PhononSpatialization.h
+++ b/Engine/Plugins/Runtime/Steam/SteamAudio/Source/SteamAudio/Private/PhononSpatialization.h
@@ -24,16 +24,19 @@ namespace SteamAudio
 		TArray<float> OutArray;
 	};
 
+	/************************************************************************/
+	/* FPhononSpatialization                                                */
+	/* Spatialization plugin using Steam Audio's HRTF spatialization library*/
+	/************************************************************************/
 	class FPhononSpatialization : public IAudioSpatialization
 	{
 	public:
 		FPhononSpatialization();
 		~FPhononSpatialization();
 
-		virtual void Initialize(const uint32 SampleRate, const uint32 NumSources, const uint32 OutputBufferLength) override;
+		virtual void Initialize(const FAudioPluginInitializationParams InitializationParams) override;
 		virtual bool IsSpatializationEffectInitialized() const override;
-		virtual void OnInitSource(const uint32 SourceId, const FName& AudioComponentUserId,
-			USpatializationPluginSourceSettingsBase* InSettings) override;
+		virtual void OnInitSource(const uint32 SourceId, const FName& AudioComponentUserId, USpatializationPluginSourceSettingsBase* InSettings) override;
 		virtual void OnReleaseSource(const uint32 SourceId) override;
 		virtual void ProcessAudio(const FAudioPluginSourceInputData& InputData, FAudioPluginSourceOutputData& OutputData) override;
 
diff --git a/Engine/Plugins/Runtime/Steam/SteamAudio/Source/SteamAudio/Private/SteamAudioEnvironment.cpp b/Engine/Plugins/Runtime/Steam/SteamAudio/Source/SteamAudio/Private/SteamAudioEnvironment.cpp
new file mode 100644
--- /dev/null
+++ b/Engine/Plugins/Runtime/Steam/SteamAudio/Source/SteamAudio/Private/SteamAudioEnvironment.cpp
@@ -0,0 +1,176 @@
+//
+// Copyright (C) Valve Corporation. All rights reserved.
+//
+
+#include "SteamAudioEnvironment.h"
+#include "PhononScene.h"
+#include "PhononProbeVolume.h"
+#include "CoreMinimal.h"
+#include "Stats/Stats.h"
+#include "Modules/ModuleInterface.h"
+#include "Modules/ModuleManager.h"
+#include "Classes/Kismet/GameplayStatics.h"
+#include "IPluginManager.h"
+#include "EngineUtils.h"
+#include "Regex.h"
+#include "ScopeLock.h"
+#include "HAL/PlatformProcess.h"
+#include "Engine/StreamableManager.h"
+#include "AudioDevice.h"
+#include "PhononScene.h"
+#include "SteamAudioSettings.h"
+#include "PhononProbeVolume.h"
+
+namespace SteamAudio
+{
+
+	FEnvironment::FEnvironment()
+		: EnvironmentCriticalSection()
+		, ComputeDevice(nullptr)
+		, PhononScene(nullptr)
+		, PhononEnvironment(nullptr)
+		, EnvironmentalRenderer(nullptr)
+		, ProbeManager(nullptr)
+		, SimulationSettings()
+		, RenderingSettings()
+		, EnvironmentalOutputAudioFormat()
+	{
+	}
+
+	FEnvironment::~FEnvironment()
+	{
+		Shutdown();
+	}
+
+	IPLhandle FEnvironment::Initialize(UWorld* World, FAudioDevice* InAudioDevice)
+	{
+		if (World == nullptr)
+		{
+			UE_LOG(LogSteamAudio, Error, TEXT("Unable to create Phonon environment: null World."));
+			return nullptr;
+		}
+
+		if (InAudioDevice == nullptr)
+		{
+			UE_LOG(LogSteamAudio, Error, TEXT("Unable to create Phonon environment: null Audio Device."));
+			return nullptr;
+		}
+
+		TArray<AActor*> PhononSceneActors;
+		UGameplayStatics::GetAllActorsOfClass(World, APhononScene::StaticClass(), PhononSceneActors);
+
+		if (PhononSceneActors.Num() == 0)
+		{
+			UE_LOG(LogSteamAudio, Error, TEXT("Unable to create Phonon environment: PhononScene not found. Be sure to add a PhononScene actor to your level and export the scene."));
+			return nullptr;
+		}
+		else if (PhononSceneActors.Num() > 1)
+		{
+			UE_LOG(LogSteamAudio, Warning, TEXT("More than one PhononScene actor found in level. Arbitrarily choosing one. Ensure only one exists to avoid unexpected behavior."));
+		}
+
+		APhononScene* PhononSceneActor = Cast<APhononScene>(PhononSceneActors[0]);
+		check(PhononSceneActor);
+
+		if (PhononSceneActor->SceneData.Num() == 0)
+		{
+			UE_LOG(LogSteamAudio, Error, TEXT("Unable to create Phonon environment: PhononScene actor does not have scene data. Be sure to export the scene."));
+			return nullptr;
+		}
+
+		int32 IndirectImpulseResponseOrder = GetDefault<USteamAudioSettings>()->IndirectImpulseResponseOrder;
+		float IndirectImpulseResponseDuration = GetDefault<USteamAudioSettings>()->IndirectImpulseResponseDuration;
+
+		SimulationSettings.maxConvolutionSources = GetDefault<USteamAudioSettings>()->MaxSources;
+		SimulationSettings.numBounces = GetDefault<USteamAudioSettings>()->RealtimeBounces;
+		SimulationSettings.numDiffuseSamples = GetDefault<USteamAudioSettings>()->RealtimeSecondaryRays;
+		SimulationSettings.numRays = GetDefault<USteamAudioSettings>()->RealtimeRays;
+		SimulationSettings.ambisonicsOrder = IndirectImpulseResponseOrder;
+		SimulationSettings.irDuration = IndirectImpulseResponseDuration;
+		SimulationSettings.sceneType = IPL_SCENETYPE_PHONON;
+
+		RenderingSettings.convolutionType = IPL_CONVOLUTIONTYPE_PHONON;
+		RenderingSettings.frameSize = InAudioDevice->GetBufferLength();
+		RenderingSettings.samplingRate = InAudioDevice->GetSampleRate();
+
+		EnvironmentalOutputAudioFormat.channelLayout = IPL_CHANNELLAYOUT_STEREO;
+		EnvironmentalOutputAudioFormat.channelLayoutType = IPL_CHANNELLAYOUTTYPE_AMBISONICS;
+		EnvironmentalOutputAudioFormat.channelOrder = IPL_CHANNELORDER_DEINTERLEAVED;
+		EnvironmentalOutputAudioFormat.numSpeakers = (IndirectImpulseResponseOrder + 1) * (IndirectImpulseResponseOrder + 1);
+		EnvironmentalOutputAudioFormat.speakerDirections = nullptr;
+		EnvironmentalOutputAudioFormat.ambisonicsOrder = IndirectImpulseResponseOrder;
+		EnvironmentalOutputAudioFormat.ambisonicsNormalization = IPL_AMBISONICSNORMALIZATION_N3D;
+		EnvironmentalOutputAudioFormat.ambisonicsOrdering = IPL_AMBISONICSORDERING_ACN;
+
+		IPLerror Error = IPLerror::IPL_STATUS_SUCCESS;
+
+		iplLoadFinalizedScene(GlobalContext, SimulationSettings, PhononSceneActor->SceneData.GetData(), PhononSceneActor->SceneData.Num(),
+			ComputeDevice, nullptr, &PhononScene);
+		LogSteamAudioStatus(Error);
+
+		Error = iplCreateProbeManager(&ProbeManager);
+		LogSteamAudioStatus(Error);
+
+		TArray<AActor*> PhononProbeVolumes;
+		UGameplayStatics::GetAllActorsOfClass(World, APhononProbeVolume::StaticClass(), PhononProbeVolumes);
+
+		for (AActor* PhononProbeVolumeActor : PhononProbeVolumes)
+		{
+			APhononProbeVolume* PhononProbeVolume = Cast<APhononProbeVolume>(PhononProbeVolumeActor);
+			IPLhandle ProbeBatch = nullptr;
+			Error = iplLoadProbeBatch(PhononProbeVolume->GetProbeBatchData(), PhononProbeVolume->GetProbeBatchDataSize(), &ProbeBatch);
+			LogSteamAudioStatus(Error);
+
+			iplAddProbeBatch(ProbeManager, ProbeBatch);
+
+			ProbeBatches.Add(ProbeBatch);
+		}
+
+		Error = iplCreateEnvironment(GlobalContext, ComputeDevice, SimulationSettings, PhononScene, ProbeManager, &PhononEnvironment);
+		LogSteamAudioStatus(Error);
+
+		Error = iplCreateEnvironmentalRenderer(GlobalContext, PhononEnvironment, RenderingSettings, EnvironmentalOutputAudioFormat,
+			nullptr, nullptr, &EnvironmentalRenderer);
+		LogSteamAudioStatus(Error);
+
+		return EnvironmentalRenderer;
+	}
+
+	void FEnvironment::Shutdown()
+	{
+		FScopeLock EnvironmentLock(&EnvironmentCriticalSection);
+		if (ProbeManager)
+		{
+			for (IPLhandle ProbeBatch : ProbeBatches)
+			{
+				iplRemoveProbeBatch(ProbeManager, ProbeBatch);
+				iplDestroyProbeBatch(&ProbeBatch);
+			}
+
+			ProbeBatches.Empty();
+
+			iplDestroyProbeManager(&ProbeManager);
+		}
+
+		if (EnvironmentalRenderer)
+		{
+			iplDestroyEnvironmentalRenderer(&EnvironmentalRenderer);
+		}
+
+		if (PhononEnvironment)
+		{
+			iplDestroyEnvironment(&PhononEnvironment);
+		}
+
+		if (PhononScene)
+		{
+			iplDestroyScene(&PhononScene);
+		}
+
+		if (ComputeDevice)
+		{
+			iplDestroyComputeDevice(&ComputeDevice);
+		}
+	}
+
+}
\ No newline at end of file
diff --git a/Engine/Plugins/Runtime/Steam/SteamAudio/Source/SteamAudio/Private/SteamAudioEnvironment.h b/Engine/Plugins/Runtime/Steam/SteamAudio/Source/SteamAudio/Private/SteamAudioEnvironment.h
new file mode 100644
--- /dev/null
+++ b/Engine/Plugins/Runtime/Steam/SteamAudio/Source/SteamAudio/Private/SteamAudioEnvironment.h
@@ -0,0 +1,43 @@
+//
+// Copyright (C) Valve Corporation. All rights reserved.
+//
+
+#pragma once
+
+#include "ISteamAudioModule.h"
+#include "phonon.h"
+#include "PhononCommon.h"
+
+namespace SteamAudio
+{
+	/************************************************************************/
+	/* FEnvironment                                                         */
+	/* This class handles an instance of the Steam Audio environment,       */
+	/* as well as the environmental renderer used by the audio plugins.     */
+	/************************************************************************/
+	class FEnvironment
+	{
+	public:
+		FEnvironment();
+		~FEnvironment();
+
+		IPLhandle Initialize(UWorld* World, FAudioDevice* InAudioDevice);
+		void Shutdown();
+
+		IPLhandle GetEnvironmentalRendererHandle() { return EnvironmentalRenderer; };
+
+		FCriticalSection* GetEnvironmentCriticalSection() { return &EnvironmentCriticalSection; };
+
+	private:
+		FCriticalSection EnvironmentCriticalSection;
+		IPLhandle ComputeDevice;
+		IPLhandle PhononScene;
+		IPLhandle PhononEnvironment;
+		IPLhandle EnvironmentalRenderer;
+		IPLhandle ProbeManager;
+		TArray<IPLhandle> ProbeBatches;
+		IPLSimulationSettings SimulationSettings;
+		IPLRenderingSettings RenderingSettings;
+		IPLAudioFormat EnvironmentalOutputAudioFormat;
+	};
+}
\ No newline at end of file
diff --git a/Engine/Plugins/Runtime/Steam/SteamAudio/Source/SteamAudio/Private/SteamAudioModule.cpp b/Engine/Plugins/Runtime/Steam/SteamAudio/Source/SteamAudio/Private/SteamAudioModule.cpp
--- a/Engine/Plugins/Runtime/Steam/SteamAudio/Source/SteamAudio/Private/SteamAudioModule.cpp
+++ b/Engine/Plugins/Runtime/Steam/SteamAudio/Source/SteamAudio/Private/SteamAudioModule.cpp
@@ -3,29 +3,7 @@
 //
 
 #include "SteamAudioModule.h"
-
-#include "PhononCommon.h"
-#include "PhononSpatialization.h"
-#include "PhononOcclusion.h"
-#include "PhononReverb.h"
-#include "PhononScene.h"
-#include "SteamAudioSettings.h"
-#include "PhononProbeVolume.h"
-#include "PhononListenerObserver.h"
-
-#include "CoreMinimal.h"
-#include "Stats/Stats.h"
-#include "Modules/ModuleInterface.h"
-#include "Modules/ModuleManager.h"
-#include "Classes/Kismet/GameplayStatics.h"
-#include "IPluginManager.h"
-#include "EngineUtils.h"
-#include "Regex.h"
-#include "ScopeLock.h"
 #include "Misc/Paths.h"
-#include "HAL/PlatformProcess.h"
-#include "Engine/StreamableManager.h"
-#include "AudioDevice.h"
 
 IMPLEMENT_MODULE(SteamAudio::FSteamAudioModule, SteamAudio)
 
@@ -36,26 +14,11 @@ namespace SteamAudio
 	static bool bModuleStartedUp = false;
 
 	FSteamAudioModule::FSteamAudioModule()
-		: SpatializationInstance(nullptr)
-		, OcclusionInstance(nullptr)
-		, ReverbInstance(nullptr)
-		, ListenerObserverInstance(nullptr)
-		, OwningAudioDevice(nullptr)
-		, ComputeDevice(nullptr)
-		, PhononScene(nullptr)
-		, PhononEnvironment(nullptr)
-		, EnvironmentalRenderer(nullptr)
-		, ProbeManager(nullptr)
 	{
 	}
 
 	FSteamAudioModule::~FSteamAudioModule()
 	{
-		check(!ComputeDevice);
-		check(!PhononScene);
-		check(!PhononEnvironment);
-		check(!EnvironmentalRenderer);
-		check(!ProbeManager);
 	}
 
 	void FSteamAudioModule::StartupModule()
@@ -66,19 +29,10 @@ namespace SteamAudio
 
 		UE_LOG(LogSteamAudio, Log, TEXT("FSteamAudioModule Startup"));
 
-		OwningAudioDevice = nullptr;
-		SpatializationInstance = nullptr;
-		OcclusionInstance = nullptr;
-		ReverbInstance = nullptr;
-		ComputeDevice = nullptr;
-		PhononScene = nullptr;
-		PhononEnvironment = nullptr;
-		EnvironmentalRenderer = nullptr;
-		ProbeManager = nullptr;
-
-		// We're overriding the IAudioSpatializationPlugin StartupModule, so we must be sure to register
-		// the modular features. Without this line, we will not see SPATIALIZATION HRTF in the editor UI.
-		IModularFeatures::Get().RegisterModularFeature(GetModularFeatureName(), this);
+		//Register the Steam Audio plugin factories
+		IModularFeatures::Get().RegisterModularFeature(FSpatializationPluginFactory::GetModularFeatureName(), &SpatializationPluginFactory);
+		IModularFeatures::Get().RegisterModularFeature(FReverbPluginFactory::GetModularFeatureName(), &ReverbPluginFactory);
+		IModularFeatures::Get().RegisterModularFeature(FOcclusionPluginFactory::GetModularFeatureName(), &OcclusionPluginFactory);
 
 		if (!FSteamAudioModule::PhononDllHandle)
 		{
@@ -108,199 +62,68 @@ namespace SteamAudio
 		}
 	}
 
-	void FSteamAudioModule::CreateEnvironment(UWorld* World, FAudioDevice* InAudioDevice)
+	IAudioPluginFactory* FSteamAudioModule::GetPluginFactory(EAudioPlugin PluginType)
 	{
-		if (World == nullptr)
-		{
-			UE_LOG(LogSteamAudio, Error, TEXT("Unable to create Phonon environment: null World."));
-			return;
-		}
-
-		if (InAudioDevice == nullptr)
-		{
-			UE_LOG(LogSteamAudio, Error, TEXT("Unable to create Phonon environment: null Audio Device."));
-			return;
-		}
-
-		check(OwningAudioDevice == nullptr);
-		OwningAudioDevice = InAudioDevice;
-
-		TArray<AActor*> PhononSceneActors;
-		UGameplayStatics::GetAllActorsOfClass(World, APhononScene::StaticClass(), PhononSceneActors);
-
-		if (PhononSceneActors.Num() == 0)
+		switch (PluginType)
 		{
-			UE_LOG(LogSteamAudio, Error, TEXT("Unable to create Phonon environment: PhononScene not found. Be sure to add a PhononScene actor to your level and export the scene."));
-			return;
+		case EAudioPlugin::SPATIALIZATION:
+			return &SpatializationPluginFactory;
+			break;
+		case EAudioPlugin::REVERB:
+			return &ReverbPluginFactory;
+			break;
+		case EAudioPlugin::OCCLUSION:
+			return &OcclusionPluginFactory;
+			break;
+		default:
+			return nullptr;
+			break;
 		}
-		else if (PhononSceneActors.Num() > 1)
-		{
-			UE_LOG(LogSteamAudio, Warning, TEXT("More than one PhononScene actor found in level. Arbitrarily choosing one. Ensure only one exists to avoid unexpected behavior."));
-		}
-
-		APhononScene* PhononSceneActor = Cast<APhononScene>(PhononSceneActors[0]);
-		check(PhononSceneActor);
-
-		if (PhononSceneActor->SceneData.Num() == 0)
-		{
-			UE_LOG(LogSteamAudio, Error, TEXT("Unable to create Phonon environment: PhononScene actor does not have scene data. Be sure to export the scene."));
-			return;
-		}
-
-		int32 IndirectImpulseResponseOrder = GetDefault<USteamAudioSettings>()->IndirectImpulseResponseOrder;
-		float IndirectImpulseResponseDuration = GetDefault<USteamAudioSettings>()->IndirectImpulseResponseDuration;
-
-		SimulationSettings.maxConvolutionSources = GetDefault<USteamAudioSettings>()->MaxSources;
-		SimulationSettings.numBounces = GetDefault<USteamAudioSettings>()->RealtimeBounces;
-		SimulationSettings.numDiffuseSamples = GetDefault<USteamAudioSettings>()->RealtimeSecondaryRays;
-		SimulationSettings.numRays = GetDefault<USteamAudioSettings>()->RealtimeRays;
-		SimulationSettings.ambisonicsOrder = IndirectImpulseResponseOrder;
-		SimulationSettings.irDuration = IndirectImpulseResponseDuration;
-		SimulationSettings.sceneType = IPL_SCENETYPE_PHONON;
-
-		RenderingSettings.convolutionType = IPL_CONVOLUTIONTYPE_PHONON;
-		RenderingSettings.frameSize = OwningAudioDevice->GetBufferLength();
-		RenderingSettings.samplingRate = OwningAudioDevice->GetSampleRate();
-
-		EnvironmentalOutputAudioFormat.channelLayout = IPL_CHANNELLAYOUT_STEREO;
-		EnvironmentalOutputAudioFormat.channelLayoutType = IPL_CHANNELLAYOUTTYPE_AMBISONICS;
-		EnvironmentalOutputAudioFormat.channelOrder = IPL_CHANNELORDER_DEINTERLEAVED;
-		EnvironmentalOutputAudioFormat.numSpeakers = (IndirectImpulseResponseOrder + 1) * (IndirectImpulseResponseOrder + 1);
-		EnvironmentalOutputAudioFormat.speakerDirections = nullptr;
-		EnvironmentalOutputAudioFormat.ambisonicsOrder = IndirectImpulseResponseOrder;
-		EnvironmentalOutputAudioFormat.ambisonicsNormalization = IPL_AMBISONICSNORMALIZATION_N3D;
-		EnvironmentalOutputAudioFormat.ambisonicsOrdering = IPL_AMBISONICSORDERING_ACN;
-
-		IPLerror Error = IPLerror::IPL_STATUS_SUCCESS;
-
-		iplLoadFinalizedScene(GlobalContext, SimulationSettings, PhononSceneActor->SceneData.GetData(), PhononSceneActor->SceneData.Num(), 
-			ComputeDevice, nullptr, &PhononScene);
-		LogSteamAudioStatus(Error);
-
-		Error = iplCreateProbeManager(&ProbeManager);
-		LogSteamAudioStatus(Error);
-
-		TArray<AActor*> PhononProbeVolumes;
-		UGameplayStatics::GetAllActorsOfClass(World, APhononProbeVolume::StaticClass(), PhononProbeVolumes);
-
-		for (AActor* PhononProbeVolumeActor : PhononProbeVolumes)
-		{
-			APhononProbeVolume* PhononProbeVolume = Cast<APhononProbeVolume>(PhononProbeVolumeActor);
-			IPLhandle ProbeBatch = nullptr;
-			Error = iplLoadProbeBatch(PhononProbeVolume->GetProbeBatchData(), PhononProbeVolume->GetProbeBatchDataSize(), &ProbeBatch);
-			LogSteamAudioStatus(Error);
-
-			iplAddProbeBatch(ProbeManager, ProbeBatch);
-
-			ProbeBatches.Add(ProbeBatch);
-		}
-
-		Error = iplCreateEnvironment(GlobalContext, ComputeDevice, SimulationSettings, PhononScene, ProbeManager, &PhononEnvironment);
-		LogSteamAudioStatus(Error);
-
-		Error = iplCreateEnvironmentalRenderer(GlobalContext, PhononEnvironment, RenderingSettings, EnvironmentalOutputAudioFormat,
-			nullptr, nullptr, &EnvironmentalRenderer);
-		LogSteamAudioStatus(Error);
-
-		FPhononOcclusion* PhononOcclusion = static_cast<FPhononOcclusion*>(OcclusionInstance.Get());
-		PhononOcclusion->SetEnvironmentalRenderer(EnvironmentalRenderer);
-
-		FPhononReverb* PhononReverb = static_cast<FPhononReverb*>(ReverbInstance.Get());
-		PhononReverb->SetEnvironmentalRenderer(EnvironmentalRenderer);
-		PhononReverb->CreateReverbEffect();
 	}
 
-	void FSteamAudioModule::DestroyEnvironment(FAudioDevice* InAudioDevice)
+	void FSteamAudioModule::RegisterAudioDevice(FAudioDevice* AudioDeviceHandle)
 	{
-		check(OwningAudioDevice == InAudioDevice);
-
-		OwningAudioDevice = nullptr;
-
-		FScopeLock Lock(&EnvironmentCriticalSection);
-
-		FPhononOcclusion* PhononOcclusion = static_cast<FPhononOcclusion*>(OcclusionInstance.Get());
-		PhononOcclusion->SetEnvironmentalRenderer(nullptr);
-
-		FPhononReverb* PhononReverb = static_cast<FPhononReverb*>(ReverbInstance.Get());
-		PhononReverb->SetEnvironmentalRenderer(nullptr);
-
-		if (ProbeManager)
-		{
-			for (IPLhandle ProbeBatch : ProbeBatches)
-			{
-				iplRemoveProbeBatch(ProbeManager, ProbeBatch);
-				iplDestroyProbeBatch(&ProbeBatch);
-			}
-
-			ProbeBatches.Empty();
-
-			iplDestroyProbeManager(&ProbeManager);
-		}
-
-		if (EnvironmentalRenderer)
-		{
-			iplDestroyEnvironmentalRenderer(&EnvironmentalRenderer);
-		}
-
-		if (PhononEnvironment)
-		{
-			iplDestroyEnvironment(&PhononEnvironment);
-		}
-
-		if (PhononScene)
-		{
-			iplDestroyScene(&PhononScene);
-		}
-
-		if (ComputeDevice)
+		if (!RegisteredAudioDevices.Contains(AudioDeviceHandle))
 		{
-			iplDestroyComputeDevice(&ComputeDevice);
+			TAudioPluginListenerPtr NewPhononPluginManager = TAudioPluginListenerPtr(new FPhononPluginManager());
+			AudioDeviceHandle->RegisterPluginListener(NewPhononPluginManager);
+			RegisteredAudioDevices.Add(AudioDeviceHandle);
 		}
 	}
 
-	TAudioSpatializationPtr FSteamAudioModule::CreateSpatializationInterface(class FAudioDevice* AudioDevice)
+	void FSteamAudioModule::UnregisterAudioDevice(FAudioDevice* AudioDeviceHandle)
 	{
-		return SpatializationInstance = TAudioSpatializationPtr(new FPhononSpatialization());
+		RegisteredAudioDevices.Remove(AudioDeviceHandle);
 	}
 
-	TAudioOcclusionPtr FSteamAudioModule::CreateOcclusionInterface(class FAudioDevice* AudioDevice)
+	TAudioOcclusionPtr FOcclusionPluginFactory::CreateNewOcclusionPlugin(FAudioDevice* OwningDevice)
 	{
-		return OcclusionInstance = TAudioOcclusionPtr(new FPhononOcclusion());
-	}
-
-	TAudioReverbPtr FSteamAudioModule::CreateReverbInterface(class FAudioDevice* AudioDevice)
-	{
-		return ReverbInstance = TAudioReverbPtr(new FPhononReverb());
-	}
-
-	TAudioListenerObserverPtr FSteamAudioModule::CreateListenerObserverInterface(class FAudioDevice* AudioDevice)
-	{
-		return ListenerObserverInstance = TAudioListenerObserverPtr(new FPhononListenerObserver());
-	}
+		//Register audio device to the Steam Module
+		FSteamAudioModule* Module = &FModuleManager::GetModuleChecked<FSteamAudioModule>("SteamAudio");
+		if (Module != nullptr)
+		{
+			Module->RegisterAudioDevice(OwningDevice);
+		}
 
-	FPhononSpatialization* FSteamAudioModule::GetSpatializationInstance() const
-	{
-		return static_cast<FPhononSpatialization*>(SpatializationInstance.Get());
+		return TAudioOcclusionPtr(new FPhononOcclusion());
 	}
 
-	FPhononOcclusion* FSteamAudioModule::GetOcclusionInstance() const
+	TAudioReverbPtr FReverbPluginFactory::CreateNewReverbPlugin(FAudioDevice* OwningDevice)
 	{
-		return static_cast<FPhononOcclusion*>(OcclusionInstance.Get());
-	}
+		//Register the audio device to the steam module:
+		FSteamAudioModule* Module = &FModuleManager::GetModuleChecked<FSteamAudioModule>("SteamAudio");
+		if (Module != nullptr)
+		{
+			Module->RegisterAudioDevice(OwningDevice);
+		}
 
-	FPhononReverb* FSteamAudioModule::GetReverbInstance() const
-	{
-		return static_cast<FPhononReverb*>(ReverbInstance.Get());
+		return TAudioReverbPtr(new FPhononReverb());
 	}
 
-	FPhononListenerObserver* FSteamAudioModule::GetListenerObserverInstance() const
+	TAudioSpatializationPtr FSpatializationPluginFactory::CreateNewSpatializationPlugin(FAudioDevice* OwningDevice)
 	{
-		return static_cast<FPhononListenerObserver*>(ListenerObserverInstance.Get());
+		return TAudioSpatializationPtr(new FPhononSpatialization());
 	}
 
-	FCriticalSection& FSteamAudioModule::GetEnvironmentCriticalSection()
-	{
-		return EnvironmentCriticalSection;
-	}
 }
 
diff --git a/Engine/Plugins/Runtime/Steam/SteamAudio/Source/SteamAudio/Private/SteamAudioModule.h b/Engine/Plugins/Runtime/Steam/SteamAudio/Source/SteamAudio/Private/SteamAudioModule.h
--- a/Engine/Plugins/Runtime/Steam/SteamAudio/Source/SteamAudio/Private/SteamAudioModule.h
+++ b/Engine/Plugins/Runtime/Steam/SteamAudio/Source/SteamAudio/Private/SteamAudioModule.h
@@ -5,10 +5,101 @@
 #pragma once
 
 #include "ISteamAudioModule.h"
-#include "phonon.h"
+#include "PhononSpatialization.h"
+#include "PhononOcclusion.h"
+#include "PhononReverb.h"
+#include "PhononPluginManager.h"
+#include "AudioPluginUtilities.h"
+#include "AudioDevice.h"
 
 namespace SteamAudio
 {
+	/************************************************************************/
+	/* Plugin Factories registered with Steam Audio Module                  */
+	/************************************************************************/
+	class FSpatializationPluginFactory : public IAudioSpatializationFactory
+	{
+	public:
+		virtual FString GetDisplayName() override
+		{
+			static FString DisplayName = FString(TEXT("Steam Audio"));
+			return DisplayName;
+		}
+
+		virtual bool SupportsPlatform(EAudioPlatform Platform) override
+		{
+			if (Platform == EAudioPlatform::Windows)
+			{
+				return true;
+			}
+			else
+			{
+				return false;
+			}
+		}
+
+		virtual TAudioSpatializationPtr CreateNewSpatializationPlugin(FAudioDevice* OwningDevice) override;
+
+		virtual bool HasCustomSpatializationSetting() const override { return true; }
+	};
+
+	class FReverbPluginFactory : public IAudioReverbFactory
+	{
+	public:
+		virtual FString GetDisplayName() override
+		{
+			static FString DisplayName = FString(TEXT("Steam Audio"));
+			return DisplayName;
+		}
+
+		virtual bool SupportsPlatform(EAudioPlatform Platform) override
+		{
+			if (Platform == EAudioPlatform::Windows)
+			{
+				return true;
+			}
+			else
+			{
+				return false;
+			}
+		}
+
+		virtual TAudioReverbPtr CreateNewReverbPlugin(FAudioDevice* OwningDevice) override;
+
+		virtual bool HasCustomReverbSetting() const override { return true; }
+	};
+
+	class FOcclusionPluginFactory : public IAudioOcclusionFactory
+	{
+	public:
+		virtual FString GetDisplayName() override
+		{
+			static FString DisplayName = FString(TEXT("Steam Audio"));
+			return DisplayName;
+		}
+
+		virtual bool SupportsPlatform(EAudioPlatform Platform) override
+		{
+			if (Platform == EAudioPlatform::Windows)
+			{
+				return true;
+			}
+			else
+			{
+				return false;
+			}
+		}
+
+		virtual TAudioOcclusionPtr CreateNewOcclusionPlugin(FAudioDevice* OwningDevice) override;
+
+		virtual bool HasCustomOcclusionSetting() const override { return true; }
+	};
+
+	/************************************************************************/
+	/* FSteamAudioModule                                                    */
+	/* Module interface for Steam Audio. Registers the Plugin Factories     */
+	/* and handles the third party DLL.                                     */
+	/************************************************************************/
 	class FSteamAudioModule : public ISteamAudioModule
 	{
 	public:
@@ -17,48 +108,34 @@ namespace SteamAudio
 
 		virtual void StartupModule() override;
 		virtual void ShutdownModule() override;
-
-		//~ Begin IAudioPlugin
-		virtual bool ImplementsSpatialization() const override { return true; }
-		virtual bool HasCustomSpatializationSetting() const override { return true; }
-		virtual bool ImplementsOcclusion() const override { return true; }
-		virtual bool HasCustomOcclusionSetting() const override { return true; }
-		virtual bool ImplementsReverb() const override { return true; }
-		virtual bool HasCustomReverbSetting() const override { return true; }
-		virtual TAudioSpatializationPtr CreateSpatializationInterface(class FAudioDevice* AudioDevice) override;
-		virtual TAudioOcclusionPtr CreateOcclusionInterface(class FAudioDevice* AudioDevice) override;
-		virtual TAudioReverbPtr CreateReverbInterface(class FAudioDevice* AudioDevice) override;
-		virtual TAudioListenerObserverPtr CreateListenerObserverInterface(class FAudioDevice* AudioDevice) override;
 		//~ End IAudioPlugin
 
-		void CreateEnvironment(UWorld* World, FAudioDevice* InAudioDevice);
-		void DestroyEnvironment(FAudioDevice* InAudioDevice);
+		/** 
+		*  Returns pointer to each type of plugin factory with Steam Audio.
+		*  returns nullptr if given an invalid parameter for PluginType.
+		*/
+		IAudioPluginFactory* GetPluginFactory(EAudioPlugin PluginType);
+
+		/*
+		 * Registers a given audio device to the Steam Audio Device.
+		 * Every audio device running Steam Audio's reverb or occlusion manager
+		 * requires a PhononPluginManager registered to it.
+		 */
+		void RegisterAudioDevice(FAudioDevice* AudioDeviceHandle);
 
-		FCriticalSection& GetEnvironmentCriticalSection();
-		class FPhononSpatialization* GetSpatializationInstance() const;
-		class FPhononOcclusion* GetOcclusionInstance() const;
-		class FPhononReverb* GetReverbInstance() const;
-		class FPhononListenerObserver* GetListenerObserverInstance() const;
+		/*
+		* Unregisters an audio device from the module.
+		*/
+		void UnregisterAudioDevice(FAudioDevice* AudioDeviceHandle);
 
 	private:
 		static void* PhononDllHandle;
 
-		TAudioSpatializationPtr SpatializationInstance;
-		TAudioOcclusionPtr OcclusionInstance;
-		TAudioReverbPtr ReverbInstance;
-		TAudioListenerObserverPtr ListenerObserverInstance;
-
-		FAudioDevice* OwningAudioDevice;
-
-		FCriticalSection EnvironmentCriticalSection;
-		IPLhandle ComputeDevice;
-		IPLhandle PhononScene;
-		IPLhandle PhononEnvironment;
-		IPLhandle EnvironmentalRenderer;
-		IPLhandle ProbeManager;
-		TArray<IPLhandle> ProbeBatches;
-		IPLSimulationSettings SimulationSettings;
-		IPLRenderingSettings RenderingSettings;
-		IPLAudioFormat EnvironmentalOutputAudioFormat;
+		TArray<FAudioDevice*> RegisteredAudioDevices;
+
+		//Factories:
+		FSpatializationPluginFactory SpatializationPluginFactory;
+		FReverbPluginFactory ReverbPluginFactory;
+		FOcclusionPluginFactory OcclusionPluginFactory;
 	};
 }
\ No newline at end of file
diff --git a/Engine/Plugins/Runtime/Steam/SteamAudio/Source/SteamAudio/Public/ISteamAudioModule.h b/Engine/Plugins/Runtime/Steam/SteamAudio/Source/SteamAudio/Public/ISteamAudioModule.h
--- a/Engine/Plugins/Runtime/Steam/SteamAudio/Source/SteamAudio/Public/ISteamAudioModule.h
+++ b/Engine/Plugins/Runtime/Steam/SteamAudio/Source/SteamAudio/Public/ISteamAudioModule.h
@@ -6,7 +6,7 @@
 
 #include "IAudioExtensionPlugin.h"
 
-class ISteamAudioModule : public IAudioPlugin
+class ISteamAudioModule : public IModuleInterface
 {
 public:
 
diff --git a/Engine/Plugins/Runtime/Synthesis/Source/Synthesis/Classes/SubmixEffects/SubmixEffectDelay.h b/Engine/Plugins/Runtime/Synthesis/Source/Synthesis/Classes/SubmixEffects/SubmixEffectDelay.h
new file mode 100644
--- /dev/null
+++ b/Engine/Plugins/Runtime/Synthesis/Source/Synthesis/Classes/SubmixEffects/SubmixEffectDelay.h
@@ -0,0 +1,131 @@
+// Copyright 1998-2017 Epic Games, Inc. All Rights Reserved.
+
+#pragma once
+
+#include "CoreMinimal.h"
+#include "AudioEffect.h"
+#include "Delay.h"
+#include "DSP/Dsp.h"
+#include "Sound/SoundEffectSubmix.h"
+#include "SubmixEffectDelay.generated.h"
+
+// ========================================================================
+// FSubmixEffectDelaySettings
+// UStruct used to define user-exposed params for use with your effect.
+// ========================================================================
+
+USTRUCT(BlueprintType)
+struct SYNTHESIS_API FSubmixEffectDelaySettings
+{
+	GENERATED_USTRUCT_BODY()
+
+	// Maximum possible length for a delay, in milliseconds. Changing this at runtime will reset the effect.
+	UPROPERTY(EditAnywhere, BlueprintReadWrite, Category = Initialization, meta = (ClampMin = "10.0", UIMin = "10.0", UIMax = "20000.0"))
+	float MaximumDelayLength;
+
+	// Number of milliseconds over which a tap will reach it's set length and gain. Smaller values are more responsive, while larger values will make pitching less dramatic.
+	UPROPERTY(EditAnywhere, BlueprintReadWrite, Category = Realtime, meta = (ClampMin = "0.0", UIMin = "0.0", UIMax = "20000.0"))
+	float InterpolationTime;
+
+	// Number of milliseconds of delay.
+	UPROPERTY(EditAnywhere, BlueprintReadWrite, Category = Realtime, meta = (ClampMin = "0.0", UIMin = "0.0", UIMax = "20000.0"))
+	float DelayLength;
+
+	FSubmixEffectDelaySettings()
+		: MaximumDelayLength(2000.0f)
+		, InterpolationTime(400.0f)
+		, DelayLength(1000.0f)
+	{
+	}
+};
+
+class SYNTHESIS_API FSubmixEffectDelay : public FSoundEffectSubmix
+{
+public:
+	FSubmixEffectDelay();
+	~FSubmixEffectDelay();
+
+	//~ Begin FSoundEffectSubmix
+	virtual void Init(const FSoundEffectSubmixInitData& InData) override;
+	virtual void OnProcessAudio(const FSoundEffectSubmixInputData& InData, FSoundEffectSubmixOutputData& OutData) override;
+	//~ End FSoundEffectSubmix
+
+	//~ Begin FSoundEffectBase
+	virtual void OnPresetChanged() override;
+	//~End FSoundEffectBase
+
+	// Sets the reverb effect parameters based from audio thread code
+	void SetEffectParameters(const FSubmixEffectDelaySettings& InTapEffectParameters);
+
+	// Set the time it takes, in milliseconds, to arrive at a new parameter.
+	void SetInterpolationTime(float Time);
+
+	// Set how long the actualy delay is, in milliseconds.
+	void SetDelayLineLength(float Length);
+
+private:
+	// This is called on the audio render thread to query the parameters.
+	void UpdateParameters();
+
+	// Called on the audio render thread when the number of channels is changed
+	void OnNumChannelsChanged(int32 NumChannels);
+
+	// Params struct used to pass parameters safely to the audio render thread.
+	Audio::TParams<FSubmixEffectDelaySettings> Params;
+
+	// Sample rate cached at initialization. Used to gage interpolation times.
+	float SampleRate;
+
+	// Current maximum delay line length, in milliseconds.
+	float MaxDelayLineLength;
+
+	// Current interpolation time, in seconds.
+	float InterpolationTime;
+
+	// Most recently set delay line length.
+	float TargetDelayLineLength;
+
+	Audio::FLinearEase InterpolationInfo;
+
+	// Delay lines for each channel.
+	TArray<Audio::FDelay> DelayLines;
+};
+
+// ========================================================================
+// USubmixEffectDelayPreset
+// Class which processes audio streams and uses parameters defined in the preset class.
+// ========================================================================
+
+UCLASS()
+class SYNTHESIS_API USubmixEffectDelayPreset : public USoundEffectSubmixPreset
+{
+	GENERATED_BODY()
+
+public:
+	EFFECT_PRESET_METHODS(SubmixEffectDelay)
+
+	// Set all tap delay setting. This will replace any dynamically added or modified taps.
+	UFUNCTION(BlueprintCallable, Category = "Audio|Effects|Delay")
+	void SetSettings(const FSubmixEffectDelaySettings& InSettings);
+
+	// Get the maximum delay possible.
+	UFUNCTION(BlueprintCallable, Category = "Audio|Effects|Delay")
+	float GetMaxDelayInMilliseconds() { return DynamicSettings.MaximumDelayLength;  };
+
+	// Set the time it takes to interpolate between parameters, in milliseconds.
+	UFUNCTION(BlueprintCallable, Category = "Audio|Effects|Delay")
+	void SetInterpolationTime(float Time);
+
+	// Set how long the delay actually is, in milliseconds.
+	UFUNCTION(BlueprintCallable, Category = "Audio|Effects|Delay")
+	void SetDelay(float Length);
+
+public:
+	virtual void OnInit() override;
+
+	UPROPERTY(EditAnywhere, BlueprintReadWrite, Category = SubmixEffectPreset, Meta = (ShowOnlyInnerProperties))
+	FSubmixEffectDelaySettings Settings;
+
+	UPROPERTY(transient)
+	FSubmixEffectDelaySettings DynamicSettings;
+};
\ No newline at end of file
diff --git a/Engine/Plugins/Runtime/Synthesis/Source/Synthesis/Classes/SubmixEffects/SubmixEffectTapDelay.h b/Engine/Plugins/Runtime/Synthesis/Source/Synthesis/Classes/SubmixEffects/SubmixEffectTapDelay.h
new file mode 100644
--- /dev/null
+++ b/Engine/Plugins/Runtime/Synthesis/Source/Synthesis/Classes/SubmixEffects/SubmixEffectTapDelay.h
@@ -0,0 +1,223 @@
+// Copyright 1998-2017 Epic Games, Inc. All Rights Reserved.
+
+#pragma once
+
+#include "CoreMinimal.h"
+#include "AudioEffect.h"
+#include "Delay.h"
+#include "DSP/Dsp.h"
+#include "Sound/SoundEffectSubmix.h"
+#include "SubmixEffectTapDelay.generated.h"
+
+DECLARE_LOG_CATEGORY_EXTERN(LogTapDelay, Log, All);
+
+UENUM()
+enum class ETapLineMode : uint8
+{
+	// Send tap audio output to a channel directly
+	SendToChannel,
+
+	// Allow tap to pan between channels based on azimuth angle
+	Panning,
+
+	// Disables the tap audio and performs a fadeout
+	Disabled
+};
+
+struct SYNTHESIS_API FTapDelayInterpolationInfo
+{
+	FTapDelayInterpolationInfo()
+	{
+	}
+
+	void Init(float SampleRate);
+	
+	void SetGainValue(float Value, float InterpolationTime);
+	float GetGainValue();
+
+	void SetLengthValue(float Value, float InterpolationTime);
+	float GetLengthValue();
+
+private:
+	Audio::FLinearEase LengthParam;
+	Audio::FLinearEase GainParam;
+};
+
+USTRUCT(BlueprintType)
+struct SYNTHESIS_API FTapDelayInfo
+{
+	GENERATED_USTRUCT_BODY()
+
+	// Whether the tap line should send directly to a channel, pan, or not produce sound at all.
+	UPROPERTY(EditAnywhere, BlueprintReadWrite, Category = Realtime)
+	ETapLineMode TapLineMode;
+
+	// Amount of time before this echo is heard in milliseconds.
+	UPROPERTY(EditAnywhere, BlueprintReadWrite, Category = Realtime, meta = (ClampMin = "0.0", ClampMax = "10000"))
+	float DelayLength;
+
+	// How loud this echo should be, in decibels.
+	UPROPERTY(EditAnywhere, BlueprintReadWrite, Category = Realtime, meta = (ClampMin = "-60.0", ClampMax = "6.0"))
+	float Gain;
+
+	// When the Tap Line Mode is set to Send To Channel, this parameter designates which channel the echo should play out of.
+	UPROPERTY(EditAnywhere, BlueprintReadWrite, Category = Realtime)
+	int32 OutputChannel;
+
+	// When the Tap Line Mode is set to Panning, this parameter designates the angle at which the echo should be panned.
+	// -90 is left, 90 is right, and 180/-180 is directly behind the listener.
+	UPROPERTY(EditAnywhere, BlueprintReadWrite, Category = Realtime, meta = (ClampMin = "-180.0", ClampMax = "180.0"))
+	float PanInDegrees;
+
+	UPROPERTY(transient)
+	int32 TapId;
+
+public:
+
+	FTapDelayInfo();
+};
+
+// ========================================================================
+// FTapDelaySubmixPresetSettings
+// UStruct used to define user-exposed params for use with your effect.
+// ========================================================================
+
+USTRUCT(BlueprintType)
+struct SYNTHESIS_API FSubmixEffectTapDelaySettings
+{
+	GENERATED_USTRUCT_BODY()
+
+	// Maximum possible length for a delay, in milliseconds. Changing this at runtime will reset the effect.
+	UPROPERTY(EditAnywhere, BlueprintReadWrite, Category = Initialization, meta = (ClampMin = "10.0", UIMin = "10.0", UIMax = "20000.0"))
+	float MaximumDelayLength;
+
+	// Number of milliseconds over which a tap will reach it's set length and gain. Smaller values are more responsive, while larger values will make pitching less dramatic.
+	UPROPERTY(EditAnywhere, BlueprintReadWrite, Category = Realtime, meta = (ClampMin = "0.0", UIMin = "0.0", UIMax = "20000.0"))
+	float InterpolationTime;
+
+	// Each tap's metadata
+	UPROPERTY(EditAnywhere, BlueprintReadWrite, Category = Realtime)
+	TArray<FTapDelayInfo> Taps;
+
+	FSubmixEffectTapDelaySettings()
+		: MaximumDelayLength(10000.0f)
+		, InterpolationTime(400.0f)
+	{
+	}
+};
+
+class SYNTHESIS_API FSubmixEffectTapDelay : public FSoundEffectSubmix
+{
+public:
+	FSubmixEffectTapDelay();
+	~FSubmixEffectTapDelay();
+
+	//~ Begin FSoundEffectSubmix
+	virtual void Init(const FSoundEffectSubmixInitData& InData) override;
+	virtual void OnProcessAudio(const FSoundEffectSubmixInputData& InData, FSoundEffectSubmixOutputData& OutData) override;
+	//~ End FSoundEffectSubmix
+
+	//~ Begin FSoundEffectBase
+	virtual void OnPresetChanged() override;
+	//~End FSoundEffectBase
+
+	// Sets the reverb effect parameters based from audio thread code
+	void SetEffectParameters(const FSubmixEffectTapDelaySettings& InTapEffectParameters);
+
+	void AddTap(int32 TapId);
+	void RemoveTap(int32 TapId);
+	void SetTap(int32 TapId, const FTapDelayInfo& DelayInfo);
+
+	// Set the interpolation time
+	void SetInterpolationTime(float Time);
+
+private:
+
+	// This is called on the audio render thread to query the parameters.
+	void UpdateParameters();
+
+	// This is called in UpdateParameters to set up our per-sample parameter interpolation.
+	void UpdateInterpolations();
+
+	// Params struct used to pass parameters safely to the audio render thread.
+	Audio::TParams<FSubmixEffectTapDelaySettings> Params;
+
+	// Sample rate cached at initialization. Used to gauge interpolation times.
+	float SampleRate;
+
+	// Current maximum delay line length, in milliseconds.
+	float MaxDelayLineLength;
+
+	// Current interpolation time, in seconds.
+	float InterpolationTime;
+
+	// Target parameters that we interpolate to. Updated directly from params.
+	TArray<FTapDelayInfo> TargetTaps;
+
+	// Current state of each tap.
+	TArray<FTapDelayInterpolationInfo> CurrentTaps;
+
+	// How many increments we have left before all CurrentTaps reach the TargetTaps
+	int32 TapIncrementsRemaining;
+
+	// Whether taps have been modified.
+	bool bSettingsModified;
+
+	Audio::FDelay DelayLine;
+};
+
+// ========================================================================
+// UTapDelaySubmixPreset
+// Class which processes audio streams and uses parameters defined in the preset class.
+// ========================================================================
+
+UCLASS()
+class SYNTHESIS_API USubmixEffectTapDelayPreset : public USoundEffectSubmixPreset
+{
+	GENERATED_BODY()
+
+public:
+	EFFECT_PRESET_METHODS(SubmixEffectTapDelay)
+
+	// Set all tap delay setting. This will replace any dynamically added or modified taps.
+	UFUNCTION(BlueprintCallable, Category = "Audio|Effects|TapDelay")
+	void SetSettings(const FSubmixEffectTapDelaySettings& InSettings);
+
+	// Adds a dynamic tap delay with the given settings. Returns the tap id. 
+	UFUNCTION(BlueprintCallable, Category = "Audio|Effects|TapDelay")
+	void AddTap(int32& TapId);
+
+	// Remove the tap from the preset.
+	UFUNCTION(BlueprintCallable, Category = "Audio|Effects|TapDelay")
+	void RemoveTap(int32 TapId);
+
+	// Modify a specific tap.
+	UFUNCTION(BlueprintCallable, Category = "Audio|Effects|TapDelay")
+	void SetTap(int32 TapId, const FTapDelayInfo& TapInfo);
+
+	// Get the current info about a specific tap.
+	UFUNCTION(BlueprintCallable, Category = "Audio|Effects|TapDelay")
+	void GetTap(int32 TapId, FTapDelayInfo& TapInfo);
+
+	// Retrieve an array of all tap ids for the submix effect.
+	UFUNCTION(BlueprintCallable, Category = "Audio|Effects|TapDelay")
+	void GetTapIds(TArray<int32>& TapIds);
+
+	// Get the maximum delay possible.
+	UFUNCTION(BlueprintCallable, Category = "Audio|Effects|TapDelay")
+	float GetMaxDelayInMilliseconds() { return DynamicSettings.MaximumDelayLength;  };
+
+	// Set the time it takes to interpolate between parameters, in milliseconds.
+	UFUNCTION(BlueprintCallable, Category = "Audio|Effects|TapDelay")
+	void SetInterpolationTime(float Time);
+
+public:
+
+	virtual void OnInit() override;
+
+	UPROPERTY(EditAnywhere, BlueprintReadWrite, Category = SubmixEffectPreset, Meta = (ShowOnlyInnerProperties))
+	FSubmixEffectTapDelaySettings Settings;
+
+	UPROPERTY(transient)
+	FSubmixEffectTapDelaySettings DynamicSettings;
+};
\ No newline at end of file
diff --git a/Engine/Plugins/Runtime/Synthesis/Source/Synthesis/Classes/SynthComponents/EpicSynth1Component.h b/Engine/Plugins/Runtime/Synthesis/Source/Synthesis/Classes/SynthComponents/EpicSynth1Component.h
--- a/Engine/Plugins/Runtime/Synthesis/Source/Synthesis/Classes/SynthComponents/EpicSynth1Component.h
+++ b/Engine/Plugins/Runtime/Synthesis/Source/Synthesis/Classes/SynthComponents/EpicSynth1Component.h
@@ -377,7 +377,7 @@ class SYNTHESIS_API UModularSynthComponent : public USynthComponent
 	virtual void Init(const int32 SampleRate) override;
 
 	// Called to generate more audio
-	virtual void OnGenerateAudio(TArray<float>& OutAudio) override;
+	virtual void OnGenerateAudio(float* OutAudio, int32 NumSamples) override;
 
 public:
 
diff --git a/Engine/Plugins/Runtime/Synthesis/Source/Synthesis/Classes/SynthComponents/SynthComponentGranulator.h b/Engine/Plugins/Runtime/Synthesis/Source/Synthesis/Classes/SynthComponents/SynthComponentGranulator.h
--- a/Engine/Plugins/Runtime/Synthesis/Source/Synthesis/Classes/SynthComponents/SynthComponentGranulator.h
+++ b/Engine/Plugins/Runtime/Synthesis/Source/Synthesis/Classes/SynthComponents/SynthComponentGranulator.h
@@ -50,7 +50,7 @@ class SYNTHESIS_API UGranularSynth : public USynthComponent
 	virtual void Init(const int32 SampleRate) override;
 
 	// Called to generate more audio
-	virtual void OnGenerateAudio(TArray<float>& OutAudio) override;
+	virtual void OnGenerateAudio(float* OutAudio, int32 NumSamples) override;
 
 	//~ Begin ActorComponent Interface.
 	virtual void OnRegister() override;
diff --git a/Engine/Plugins/Runtime/Synthesis/Source/Synthesis/Private/EpicSynth1Component.cpp b/Engine/Plugins/Runtime/Synthesis/Source/Synthesis/Private/EpicSynth1Component.cpp
--- a/Engine/Plugins/Runtime/Synthesis/Source/Synthesis/Private/EpicSynth1Component.cpp
+++ b/Engine/Plugins/Runtime/Synthesis/Source/Synthesis/Private/EpicSynth1Component.cpp
@@ -603,9 +603,9 @@ void UModularSynthComponent::Init(const int32 SampleRate)
 	SetChorusEnabled(false);
 }
 
-void UModularSynthComponent::OnGenerateAudio(TArray<float>& OutAudio)
+void UModularSynthComponent::OnGenerateAudio(float* OutAudio, int32 NumSamples)
 {
-	const int32 NumFrames = OutAudio.Num() / NumChannels;
+	const int32 NumFrames = NumSamples / NumChannels;
 
 	float LeftSample = 0.0f;
 	float RightSample = 0.0f;
diff --git a/Engine/Plugins/Runtime/Synthesis/Source/Synthesis/Private/Flexiverb.cpp b/Engine/Plugins/Runtime/Synthesis/Source/Synthesis/Private/Flexiverb.cpp
--- a/Engine/Plugins/Runtime/Synthesis/Source/Synthesis/Private/Flexiverb.cpp
+++ b/Engine/Plugins/Runtime/Synthesis/Source/Synthesis/Private/Flexiverb.cpp
@@ -79,8 +79,9 @@ namespace Audio
 	}
 
 	void FFlexiverb::UpdateComplexity(const int32 Complexity)
-	{	
-		DelayLines.Init(FDelayAPF(), Complexity);
+	{
+		ScatteringMatrixLength = Complexity * 2 + 1;
+		DelayLines.Init(FDelayAPF(), ScatteringMatrixLength);
 		int32 Index = 0;
 		for (FDelayAPF& Delay : DelayLines)
 		{
@@ -88,9 +89,9 @@ namespace Audio
 			Delay.SetDelaySamples(DelayLineSampleLengths[Index++]);
 		}
 
-		DampeningArray.Init(FOnePoleLPF(), Complexity);
+		DampeningArray.Init(FOnePoleLPF(), ScatteringMatrixLength);
 
-		ScatteringMatrixLength = Complexity;
+		
 	}
 
 	void FFlexiverb::ProcessAudioFrame(const float* InBuffer, const int32 InChannels, float* OutBuffer, const int32 OutChannels)
diff --git a/Engine/Plugins/Runtime/Synthesis/Source/Synthesis/Private/SubmixEffectFlexiverb.cpp b/Engine/Plugins/Runtime/Synthesis/Source/Synthesis/Private/SubmixEffectFlexiverb.cpp
--- a/Engine/Plugins/Runtime/Synthesis/Source/Synthesis/Private/SubmixEffectFlexiverb.cpp
+++ b/Engine/Plugins/Runtime/Synthesis/Source/Synthesis/Private/SubmixEffectFlexiverb.cpp
@@ -45,9 +45,12 @@ void FSubmixEffectFlexiverb::OnProcessAudio(const FSoundEffectSubmixInputData& I
 
 	const float* InputBuffer = InData.AudioBuffer->GetData();
 	float* OutputBuffer = OutData.AudioBuffer->GetData();
-	for (int32 Sample = 0; Sample < OutData.AudioBuffer->Num(); Sample += OutData.NumChannels)
+
+	int32 InputBufferIndex = 0;
+	for (int32 OutputBufferIndex = 0; OutputBufferIndex < OutData.AudioBuffer->Num(); OutputBufferIndex += OutData.NumChannels)
 	{
-		Flexiverb.ProcessAudioFrame(&(InputBuffer[Sample]), InData.NumChannels, &(OutputBuffer[Sample]), OutData.NumChannels);
+		Flexiverb.ProcessAudioFrame(&(InputBuffer[InputBufferIndex]), InData.NumChannels, &(OutputBuffer[OutputBufferIndex]), OutData.NumChannels);
+		InputBuffer += InData.NumChannels;
 	}
 }
 
diff --git a/Engine/Plugins/Runtime/Synthesis/Source/Synthesis/Private/SynthComponentGranulator.cpp b/Engine/Plugins/Runtime/Synthesis/Source/Synthesis/Private/SynthComponentGranulator.cpp
--- a/Engine/Plugins/Runtime/Synthesis/Source/Synthesis/Private/SynthComponentGranulator.cpp
+++ b/Engine/Plugins/Runtime/Synthesis/Source/Synthesis/Private/SynthComponentGranulator.cpp
@@ -22,9 +22,9 @@ void UGranularSynth::Init(const int32 SampleRate)
 	SoundWaveLoader.Init(GetAudioDevice());
 }
 
-void UGranularSynth::OnGenerateAudio(TArray<float>& OutAudio)
+void UGranularSynth::OnGenerateAudio(float* OutAudio, int32 NumSamples)
 {
-	const int32 NumFrames = OutAudio.Num() / NumChannels;
+	const int32 NumFrames = NumSamples / NumChannels;
 	GranularSynth.Generate(OutAudio, NumFrames);
 }
 
diff --git a/Engine/Plugins/Runtime/Synthesis/Source/Synthesis/Private/SynthComponentWaveTable.cpp b/Engine/Plugins/Runtime/Synthesis/Source/Synthesis/Private/SynthComponentWaveTable.cpp
--- a/Engine/Plugins/Runtime/Synthesis/Source/Synthesis/Private/SynthComponentWaveTable.cpp
+++ b/Engine/Plugins/Runtime/Synthesis/Source/Synthesis/Private/SynthComponentWaveTable.cpp
@@ -144,7 +144,7 @@ void USynthSamplePlayer::TickComponent(float DeltaTime, enum ELevelTick TickType
 	OnSamplePlaybackProgress.Broadcast(GetCurrentPlaybackProgressTime(), GetCurrentPlaybackProgressPercent());
 }
 
-void USynthSamplePlayer::OnGenerateAudio(TArray<float>& OutAudio)
+void USynthSamplePlayer::OnGenerateAudio(float* OutAudio, int32 NumSamples)
 {
 	if (SampleBuffer.GetData() && !SampleBufferReader.HasBuffer())
 	{
@@ -158,13 +158,13 @@ void USynthSamplePlayer::OnGenerateAudio(TArray<float>& OutAudio)
 
 	if (SampleBufferReader.HasBuffer())
 	{
-		const int32 NumFrames = OutAudio.Num() / NumChannels;
+		const int32 NumFrames = NumSamples / NumChannels;
 		SampleBufferReader.Generate(OutAudio, NumFrames, NumChannels, true);
 		SamplePlaybackProgressSec = SampleBufferReader.GetPlaybackProgress();
 	}
 	else
 	{
-		for (int32 Sample = 0; Sample < OutAudio.Num(); ++Sample)
+		for (int32 Sample = 0; Sample < NumSamples; ++Sample)
 		{
 			OutAudio[Sample] = 0.0f;
 		}
diff --git a/Engine/Source/Developer/Android/AndroidPlatformEditor/AndroidPlatformEditor.Build.cs b/Engine/Source/Developer/Android/AndroidPlatformEditor/AndroidPlatformEditor.Build.cs
--- a/Engine/Source/Developer/Android/AndroidPlatformEditor/AndroidPlatformEditor.Build.cs
+++ b/Engine/Source/Developer/Android/AndroidPlatformEditor/AndroidPlatformEditor.Build.cs
@@ -28,7 +28,8 @@ public class AndroidPlatformEditor : ModuleRules
                 "RenderCore",
                 "MaterialShaderQualitySettings",
 				"MainFrame",
-			}
+                "AudioSettingsEditor"
+            }
 		);
 
 		PrivateIncludePathModuleNames.AddRange(
diff --git a/Engine/Source/Developer/Android/AndroidPlatformEditor/Private/AndroidTargetSettingsCustomization.cpp b/Engine/Source/Developer/Android/AndroidPlatformEditor/Private/AndroidTargetSettingsCustomization.cpp
--- a/Engine/Source/Developer/Android/AndroidPlatformEditor/Private/AndroidTargetSettingsCustomization.cpp
+++ b/Engine/Source/Developer/Android/AndroidPlatformEditor/Private/AndroidTargetSettingsCustomization.cpp
@@ -86,6 +86,7 @@ void FAndroidTargetSettingsCustomization::CustomizeDetails(IDetailLayoutBuilder&
 	BuildLaunchImageSection(DetailLayout);
 	BuildDaydreamAppTileImageSection(DetailLayout);
 	BuildGraphicsDebuggerSection(DetailLayout);
+	AudioPluginWidgetManager.BuildAudioCategory(DetailLayout, EAudioPlatform::Android);
 }
 
 static void OnBrowserLinkClicked(const FSlateHyperlinkRun::FMetadata& Metadata)
diff --git a/Engine/Source/Developer/Android/AndroidPlatformEditor/Private/AndroidTargetSettingsCustomization.h b/Engine/Source/Developer/Android/AndroidPlatformEditor/Private/AndroidTargetSettingsCustomization.h
--- a/Engine/Source/Developer/Android/AndroidPlatformEditor/Private/AndroidTargetSettingsCustomization.h
+++ b/Engine/Source/Developer/Android/AndroidPlatformEditor/Private/AndroidTargetSettingsCustomization.h
@@ -6,6 +6,8 @@
 #include "Misc/Attribute.h"
 #include "Input/Reply.h"
 #include "IDetailCustomization.h"
+#include "PropertyHandle.h"
+#include "TargetPlatformAudioCustomization.h"
 
 class IDetailLayoutBuilder;
 
@@ -82,5 +84,6 @@ private:
 	// Is the App ID string writable?
 	TAttribute<bool> SetupForGooglePlayAttribute;
 
+	FAudioPluginWidgetManager AudioPluginWidgetManager;
  	IDetailLayoutBuilder* SavedLayoutBuilder;
 };
diff --git a/Engine/Source/Developer/AssetTools/Private/AssetTypeActions/AssetTypeActions_PhysicsAsset.cpp b/Engine/Source/Developer/AssetTools/Private/AssetTypeActions/AssetTypeActions_PhysicsAsset.cpp
--- a/Engine/Source/Developer/AssetTools/Private/AssetTypeActions/AssetTypeActions_PhysicsAsset.cpp
+++ b/Engine/Source/Developer/AssetTools/Private/AssetTypeActions/AssetTypeActions_PhysicsAsset.cpp
@@ -1,8 +1,9 @@
 // Copyright 1998-2017 Epic Games, Inc. All Rights Reserved.
 
 #include "AssetTypeActions/AssetTypeActions_PhysicsAsset.h"
-#include "Editor/PhAT/Public/PhATModule.h"
+#include "Editor/PhysicsAssetEditor/Public/PhysicsAssetEditorModule.h"
 #include "PhysicsEngine/PhysicsAsset.h"
+#include "ThumbnailRendering/SceneThumbnailInfo.h"
 
 #define LOCTEXT_NAMESPACE "AssetTypeActions"
 
@@ -11,6 +12,19 @@ UClass* FAssetTypeActions_PhysicsAsset::GetSupportedClass() const
 	return UPhysicsAsset::StaticClass();
 }
 
+UThumbnailInfo* FAssetTypeActions_PhysicsAsset::GetThumbnailInfo(UObject* Asset) const
+{
+	UPhysicsAsset* PhysicsAsset = CastChecked<UPhysicsAsset>(Asset);
+	UThumbnailInfo* ThumbnailInfo = PhysicsAsset->ThumbnailInfo;
+	if (ThumbnailInfo == NULL)
+	{
+		ThumbnailInfo = NewObject<USceneThumbnailInfo>(PhysicsAsset, NAME_None, RF_Transactional);
+		PhysicsAsset->ThumbnailInfo = ThumbnailInfo;
+	}
+
+	return ThumbnailInfo;
+}
+
 void FAssetTypeActions_PhysicsAsset::OpenAssetEditor( const TArray<UObject*>& InObjects, TSharedPtr<IToolkitHost> EditWithinLevelEditor )
 {
 	EToolkitMode::Type Mode = EditWithinLevelEditor.IsValid() ? EToolkitMode::WorldCentric : EToolkitMode::Standalone;
@@ -21,8 +35,8 @@ void FAssetTypeActions_PhysicsAsset::OpenAssetEditor( const TArray<UObject*>& In
 		if (PhysicsAsset != NULL)
 		{
 			
-			IPhATModule* PhATModule = &FModuleManager::LoadModuleChecked<IPhATModule>( "PhAT" );
-			PhATModule->CreatePhAT(Mode, EditWithinLevelEditor, PhysicsAsset);
+			IPhysicsAssetEditorModule* PhysicsAssetEditorModule = &FModuleManager::LoadModuleChecked<IPhysicsAssetEditorModule>( "PhysicsAssetEditor" );
+			PhysicsAssetEditorModule->CreatePhysicsAssetEditor(Mode, EditWithinLevelEditor, PhysicsAsset);
 		}
 	}
 }
diff --git a/Engine/Source/Developer/AssetTools/Private/AssetTypeActions/AssetTypeActions_PhysicsAsset.h b/Engine/Source/Developer/AssetTools/Private/AssetTypeActions/AssetTypeActions_PhysicsAsset.h
--- a/Engine/Source/Developer/AssetTools/Private/AssetTypeActions/AssetTypeActions_PhysicsAsset.h
+++ b/Engine/Source/Developer/AssetTools/Private/AssetTypeActions/AssetTypeActions_PhysicsAsset.h
@@ -15,4 +15,5 @@ public:
 	virtual UClass* GetSupportedClass() const override;
 	virtual void OpenAssetEditor( const TArray<UObject*>& InObjects, TSharedPtr<class IToolkitHost> EditWithinLevelEditor = TSharedPtr<IToolkitHost>() ) override;
 	virtual uint32 GetCategories() override { return EAssetTypeCategories::Physics; }
+	virtual UThumbnailInfo* GetThumbnailInfo(UObject* Asset) const override;
 };
diff --git a/Engine/Source/Developer/AssetTools/Private/AssetTypeActions/AssetTypeActions_SkeletalMesh.cpp b/Engine/Source/Developer/AssetTools/Private/AssetTypeActions/AssetTypeActions_SkeletalMesh.cpp
--- a/Engine/Source/Developer/AssetTools/Private/AssetTypeActions/AssetTypeActions_SkeletalMesh.cpp
+++ b/Engine/Source/Developer/AssetTools/Private/AssetTypeActions/AssetTypeActions_SkeletalMesh.cpp
@@ -22,7 +22,7 @@
 #include "AssetTools.h"
 #include "AssetRegistryModule.h"
 #include "SSkeletonWidget.h"
-#include "Editor/PhAT/Public/PhATModule.h"
+#include "Editor/PhysicsAssetEditor/Public/PhysicsAssetEditorModule.h"
 #include "PersonaModule.h"
 #include "ContentBrowserModule.h"
 #include "AnimationEditorUtils.h"
@@ -34,6 +34,7 @@
 #include "ISkeletalMeshEditorModule.h"
 #include "ApexClothingUtils.h"
 #include "Algo/Transform.h"
+#include "Factories/PhysicsAssetFactory.h"
 
 #define LOCTEXT_NAMESPACE "AssetTypeActions"
 
@@ -593,14 +594,26 @@ void FAssetTypeActions_SkeletalMesh::GetPhysicsAssetMenu(FMenuBuilder& MenuBuild
 
 void FAssetTypeActions_SkeletalMesh::ExecuteNewPhysicsAsset(TArray<TWeakObjectPtr<USkeletalMesh>> Objects, bool bSetAssetToMesh)
 {
+	TArray<UObject*> CreatedObjects;
+
 	for (auto ObjIt = Objects.CreateConstIterator(); ObjIt; ++ObjIt)
 	{
 		auto Object = (*ObjIt).Get();
 		if ( Object )
 		{
-			CreatePhysicsAssetFromMesh(Object, bSetAssetToMesh);
+			if(UObject* PhysicsAsset = UPhysicsAssetFactory::CreatePhysicsAssetFromMesh(NAME_None, nullptr, Object, bSetAssetToMesh))
+			{
+				CreatedObjects.Add(PhysicsAsset);
+			}
 		}
 	}
+
+	if(CreatedObjects.Num() > 0)
+	{
+		FContentBrowserModule& ContentBrowserModule = FModuleManager::LoadModuleChecked<FContentBrowserModule>("ContentBrowser");
+		ContentBrowserModule.Get().SyncBrowserToAssets(CreatedObjects);
+		FAssetEditorManager::Get().OpenEditorForAssets(CreatedObjects);
+	}
 }
 
 void FAssetTypeActions_SkeletalMesh::ExecuteNewSkeleton(TArray<TWeakObjectPtr<USkeletalMesh>> Objects)
@@ -746,65 +759,6 @@ void FAssetTypeActions_SkeletalMesh::FillSkeletonMenu(FMenuBuilder& MenuBuilder,
 		);
 }
 
-void FAssetTypeActions_SkeletalMesh::CreatePhysicsAssetFromMesh(USkeletalMesh* SkelMesh, bool bSetToMesh) const
-{
-	// Get a unique package and asset name
-	FString Name;
-	FString PackageName;
-	CreateUniqueAssetName(SkelMesh->GetOutermost()->GetName(), TEXT("_Physics"), PackageName, Name);
-
-	// Then find/create it.
-	UPackage* Package = CreatePackage(NULL, *PackageName);
-	if ( !ensure(Package) )
-	{
-		// There was a problem creating the package
-		return;
-	}
-
-	IPhATModule* PhATModule = &FModuleManager::LoadModuleChecked<IPhATModule>( "PhAT" );
-	FPhysAssetCreateParams NewBodyData;
-	EAppReturnType::Type NewBodyResponse;
-
-	// Now show the 'asset creation' options dialog
-	PhATModule->OpenNewBodyDlg(&NewBodyData, &NewBodyResponse);
-	bool bWasOkClicked = (NewBodyResponse == EAppReturnType::Ok);
-
-	if( bWasOkClicked )
-	{			
-		UPhysicsAsset* NewAsset = NewObject<UPhysicsAsset>(Package, *Name, RF_Public | RF_Standalone | RF_Transactional);
-		if(NewAsset)
-		{
-			// Do automatic asset generation.
-			FText ErrorMessage;
-			bool bSuccess = FPhysicsAssetUtils::CreateFromSkeletalMesh(NewAsset, SkelMesh, NewBodyData, ErrorMessage, bSetToMesh);
-			if(bSuccess)
-			{
-				NewAsset->MarkPackageDirty();
-				PhATModule->CreatePhAT(EToolkitMode::Standalone, TSharedPtr<IToolkitHost>(), NewAsset);
-
-				// Notify the asset registry
-				FAssetRegistryModule::AssetCreated(NewAsset);
-
-				if(bSetToMesh)
-				{
-					// auto-link source skelmesh to the new physasset and recreate physics state if needed
-					RefreshSkelMeshOnPhysicsAssetChange(SkelMesh);
-					SkelMesh->MarkPackageDirty();
-				}
-			}
-			else
-			{
-				FMessageDialog::Open(EAppMsgType::Ok, ErrorMessage);
-				NewAsset->ClearFlags( RF_Public| RF_Standalone );
-			}
-		}
-		else
-		{
-			FMessageDialog::Open( EAppMsgType::Ok, NSLOCTEXT("CreatePhysicsAsset", "CreatePhysicsAssetFailed", "Failed to create new Physics Asset.") );
-		}
-	}
-}
-
 void FAssetTypeActions_SkeletalMesh::AssignSkeletonToMesh(USkeletalMesh* SkelMesh) const
 {
 	// Create a skeleton asset from the selected skeletal mesh. Defaults to being in the same package/group as the skeletal mesh.
diff --git a/Engine/Source/Developer/AssetTools/Private/AssetTypeActions/AssetTypeActions_SkeletalMesh.h b/Engine/Source/Developer/AssetTools/Private/AssetTypeActions/AssetTypeActions_SkeletalMesh.h
--- a/Engine/Source/Developer/AssetTools/Private/AssetTypeActions/AssetTypeActions_SkeletalMesh.h
+++ b/Engine/Source/Developer/AssetTools/Private/AssetTypeActions/AssetTypeActions_SkeletalMesh.h
@@ -54,9 +54,6 @@ private:
 
 	// Helper functions
 private:
-	/** Creates a physics asset based on the mesh */
-	void CreatePhysicsAssetFromMesh(USkeletalMesh* SkelMesh, bool bSetToMesh) const;
-
 	/** Assigns a skeleton to the mesh */
 	void AssignSkeletonToMesh(USkeletalMesh* SkelMesh) const;
 
diff --git a/Engine/Source/Developer/AudioSettingsEditor/Public/TargetPlatformAudioCustomization.h b/Engine/Source/Developer/AudioSettingsEditor/Public/TargetPlatformAudioCustomization.h
new file mode 100644
--- /dev/null
+++ b/Engine/Source/Developer/AudioSettingsEditor/Public/TargetPlatformAudioCustomization.h
@@ -0,0 +1,40 @@
+#pragma once
+
+#include "CoreMinimal.h"
+#include "PropertyHandle.h"
+#include "DetailLayoutBuilder.h"
+
+#if WITH_ENGINE
+#include "AudioPluginUtilities.h"
+#endif
+
+class AUDIOSETTINGSEDITOR_API FAudioPluginWidgetManager
+{
+public:
+	FAudioPluginWidgetManager();
+	~FAudioPluginWidgetManager();
+
+	/* Builds out the audio category for a specific audio section for a platform settings page. */
+	void BuildAudioCategory(IDetailLayoutBuilder& DetailLayout, EAudioPlatform AudioPlatform);
+
+	/** Creates widget from a scan of loaded audio plugins for an individual plugin type. */
+	TSharedRef<SWidget> MakeAudioPluginSelectorWidget(const TSharedPtr<IPropertyHandle>& PropertyHandle, EAudioPlugin AudioPluginType, EAudioPlatform AudioPlatform);
+
+private:
+	/** Handles when a new plugin is selected. */
+	static void OnPluginSelected(FString PluginName, TSharedPtr<IPropertyHandle> PropertyHandle);
+
+	void OnPluginTextCommitted(const FText& InText, ETextCommit::Type CommitType, EAudioPlugin AudioPluginType, TSharedPtr<IPropertyHandle> PropertyHandle);
+
+	FText OnGetPluginText(EAudioPlugin AudioPluginType);
+
+private:
+	/** Cached references to text for Spatialization, Reverb and Occlusion settings */
+	TSharedPtr<FText> SelectedReverb;
+	TSharedPtr<FText> SelectedSpatialization;
+	TSharedPtr<FText> SelectedOcclusion;
+
+	TArray<TSharedPtr<FText>> SpatializationPlugins;
+	TArray<TSharedPtr<FText>> ReverbPlugins;
+	TArray<TSharedPtr<FText>> OcclusionPlugins;
+};
diff --git a/Engine/Source/Developer/HTML5/HTML5PlatformEditor/Classes/HTML5TargetSettings.h b/Engine/Source/Developer/HTML5/HTML5PlatformEditor/Classes/HTML5TargetSettings.h
--- a/Engine/Source/Developer/HTML5/HTML5PlatformEditor/Classes/HTML5TargetSettings.h
+++ b/Engine/Source/Developer/HTML5/HTML5PlatformEditor/Classes/HTML5TargetSettings.h
@@ -136,5 +136,17 @@ public:
 	*/
 	UPROPERTY(GlobalConfig, EditAnywhere, Category = Amazon_S3, Meta = (DisplayName = "Nested Folder Name", EditCondition = "UploadToS3"))
 	FString S3FolderName;
+
+	/** Which of the currently enabled spatialization plugins to use on Windows. */
+	UPROPERTY(config, EditAnywhere, Category = "Audio")
+	FString SpatializationPlugin;
+
+	/** Which of the currently enabled reverb plugins to use on Windows. */
+	UPROPERTY(config, EditAnywhere, Category = "Audio")
+	FString ReverbPlugin;
+
+	/** Which of the currently enabled occlusion plugins to use on Windows. */
+	UPROPERTY(config, EditAnywhere, Category = "Audio")
+	FString OcclusionPlugin;
 };
 
diff --git a/Engine/Source/Developer/HTML5/HTML5PlatformEditor/HTML5PlatformEditor.Build.cs b/Engine/Source/Developer/HTML5/HTML5PlatformEditor/HTML5PlatformEditor.Build.cs
--- a/Engine/Source/Developer/HTML5/HTML5PlatformEditor/HTML5PlatformEditor.Build.cs
+++ b/Engine/Source/Developer/HTML5/HTML5PlatformEditor/HTML5PlatformEditor.Build.cs
@@ -27,6 +27,7 @@ public class HTML5PlatformEditor : ModuleRules
 				// WITH_EDITOR only?
 				"UnrealEd",
 				"DesktopPlatform",
+                "AudioSettingsEditor"
 			}
 		);
 
diff --git a/Engine/Source/Developer/HTML5/HTML5PlatformEditor/Private/HTML5TargetSettingsCustomization.h b/Engine/Source/Developer/HTML5/HTML5PlatformEditor/Private/HTML5TargetSettingsCustomization.h
--- a/Engine/Source/Developer/HTML5/HTML5PlatformEditor/Private/HTML5TargetSettingsCustomization.h
+++ b/Engine/Source/Developer/HTML5/HTML5PlatformEditor/Private/HTML5TargetSettingsCustomization.h
@@ -7,6 +7,7 @@
 #include "Misc/Attribute.h"
 #include "Input/Reply.h"
 #include "IDetailCustomization.h"
+#include "TargetPlatformAudioCustomization.h"
 
 class IPropertyHandle;
 class SEditableTextBox;
@@ -30,4 +31,7 @@ public:
 
 private:
 	FHTML5TargetSettingsCustomization();
+
+	//Audio plugin widget builder:
+	FAudioPluginWidgetManager AudioPluginWidgetManager;
 };
diff --git a/Engine/Source/Developer/HierarchicalLODUtilities/Private/HierarchicalLODUtilities.cpp b/Engine/Source/Developer/HierarchicalLODUtilities/Private/HierarchicalLODUtilities.cpp
--- a/Engine/Source/Developer/HierarchicalLODUtilities/Private/HierarchicalLODUtilities.cpp
+++ b/Engine/Source/Developer/HierarchicalLODUtilities/Private/HierarchicalLODUtilities.cpp
@@ -405,8 +405,8 @@ void FHierarchicalLODUtilities::DestroyClusterData(ALODActor* InActor)
 #if WITH_EDITOR
 			// Close possible open editors using this asset	
 			FAssetEditorManager::Get().CloseAllEditorsForAsset(AssetObject);
+			InActor->PreviousSubObjects.Add(AssetObject);
 #endif // WITH_EDITOR
-			AssetObject->ClearFlags(RF_AllFlags);
 		}
 	}
 
diff --git a/Engine/Source/Developer/IOS/IOSPlatformEditor/IOSPlatformEditor.Build.cs b/Engine/Source/Developer/IOS/IOSPlatformEditor/IOSPlatformEditor.Build.cs
--- a/Engine/Source/Developer/IOS/IOSPlatformEditor/IOSPlatformEditor.Build.cs
+++ b/Engine/Source/Developer/IOS/IOSPlatformEditor/IOSPlatformEditor.Build.cs
@@ -26,7 +26,8 @@ public class IOSPlatformEditor : ModuleRules
 				"TargetPlatform",
 				"MaterialShaderQualitySettings",
 				"RenderCore",
-			}
+                "AudioSettingsEditor"
+            }
 		);
 
 		PrivateIncludePathModuleNames.AddRange(
diff --git a/Engine/Source/Developer/IOS/IOSPlatformEditor/Private/IOSTargetSettingsCustomization.cpp b/Engine/Source/Developer/IOS/IOSPlatformEditor/Private/IOSTargetSettingsCustomization.cpp
--- a/Engine/Source/Developer/IOS/IOSPlatformEditor/Private/IOSTargetSettingsCustomization.cpp
+++ b/Engine/Source/Developer/IOS/IOSPlatformEditor/Private/IOSTargetSettingsCustomization.cpp
@@ -120,6 +120,7 @@ void FIOSTargetSettingsCustomization::CustomizeDetails(IDetailLayoutBuilder& Det
 
 	BuildRemoteBuildingSection(DetailLayout);
 
+	AudioPluginWidgetManager.BuildAudioCategory(DetailLayout, EAudioPlatform::IOS);
 	FindRequiredFiles();
 }
 
diff --git a/Engine/Source/Developer/IOS/IOSPlatformEditor/Private/IOSTargetSettingsCustomization.h b/Engine/Source/Developer/IOS/IOSPlatformEditor/Private/IOSTargetSettingsCustomization.h
--- a/Engine/Source/Developer/IOS/IOSPlatformEditor/Private/IOSTargetSettingsCustomization.h
+++ b/Engine/Source/Developer/IOS/IOSPlatformEditor/Private/IOSTargetSettingsCustomization.h
@@ -10,6 +10,7 @@
 #include "Widgets/Views/STableRow.h"
 #include "IDetailCustomization.h"
 #include "ShaderFormatsPropertyDetails.h"
+#include "TargetPlatformAudioCustomization.h"
 
 class FMonitoredProcess;
 class IDetailLayoutBuilder;
@@ -225,4 +226,6 @@ private:
 
 	/** Reference to the os version property warning text box. */
 	TSharedPtr< SErrorText > GLVersionWarningTextBox;
+	/** Structure used to manage audio plugin platform settings */
+	FAudioPluginWidgetManager AudioPluginWidgetManager;
 };
diff --git a/Engine/Source/Developer/Linux/LinuxTargetPlatform/Classes/LinuxTargetSettings.h b/Engine/Source/Developer/Linux/LinuxTargetPlatform/Classes/LinuxTargetSettings.h
--- a/Engine/Source/Developer/Linux/LinuxTargetPlatform/Classes/LinuxTargetSettings.h
+++ b/Engine/Source/Developer/Linux/LinuxTargetPlatform/Classes/LinuxTargetSettings.h
@@ -13,7 +13,7 @@
 
 
 /**
- * Implements the settings for the Android target platform.
+ * Implements the settings for the Linux target platform.
  */
 UCLASS(config=Engine, defaultconfig)
 class LINUXTARGETPLATFORM_API ULinuxTargetSettings
@@ -23,6 +23,18 @@ public:
 
 	GENERATED_UCLASS_BODY()
 
+	/** Which of the currently enabled spatialization plugins to use on Windows. */
+	UPROPERTY(config, EditAnywhere, Category = "Audio")
+	FString SpatializationPlugin;
+
+	/** Which of the currently enabled reverb plugins to use on Windows. */
+	UPROPERTY(config, EditAnywhere, Category = "Audio")
+	FString ReverbPlugin;
+
+	/** Which of the currently enabled occlusion plugins to use on Windows. */
+	UPROPERTY(config, EditAnywhere, Category = "Audio")
+	FString OcclusionPlugin;
+
 	/**
 	 * The collection of RHI's we want to support on this platform.
 	 * This is not always the full list of RHI we can support.
diff --git a/Engine/Source/Developer/Mac/MacTargetPlatform/Classes/MacTargetSettings.h b/Engine/Source/Developer/Mac/MacTargetPlatform/Classes/MacTargetSettings.h
--- a/Engine/Source/Developer/Mac/MacTargetPlatform/Classes/MacTargetSettings.h
+++ b/Engine/Source/Developer/Mac/MacTargetPlatform/Classes/MacTargetSettings.h
@@ -82,4 +82,16 @@ public:
 	/** The number of workers to use to compute source audio. Will only use up to the max number of sources. Will evenly divide sources to each source worker. */
 	UPROPERTY(config, EditAnywhere, Category = "Audio", meta = (ClampMin = "0", UIMin = "0", DisplayName = "Number of Source Workers"))
 	int32 AudioNumSourceWorkers;
+
+	/** Which of the currently enabled spatialization plugins to use on Windows. */
+	UPROPERTY(config, EditAnywhere, Category = "Audio")
+	FString SpatializationPlugin;
+
+	/** Which of the currently enabled reverb plugins to use on Windows. */
+	UPROPERTY(config, EditAnywhere, Category = "Audio")
+	FString ReverbPlugin;
+
+	/** Which of the currently enabled occlusion plugins to use on Windows. */
+	UPROPERTY(config, EditAnywhere, Category = "Audio")
+	FString OcclusionPlugin;
 };
diff --git a/Engine/Source/Developer/MaterialBaking/Private/MaterialOptionsCustomization.cpp b/Engine/Source/Developer/MaterialBaking/Private/MaterialOptionsCustomization.cpp
--- a/Engine/Source/Developer/MaterialBaking/Private/MaterialOptionsCustomization.cpp
+++ b/Engine/Source/Developer/MaterialBaking/Private/MaterialOptionsCustomization.cpp
@@ -12,6 +12,9 @@
 #include "Widgets/Text/STextBlock.h"
 #include "RHI.h"
 
+#include "IPropertyTypeCustomization.h"
+#include "IPropertyUtilities.h"
+
 TSharedRef<IPropertyTypeCustomization> FPropertyEntryCustomization::MakeInstance()
 {
 	return MakeShareable(new FPropertyEntryCustomization);
@@ -20,6 +23,7 @@ TSharedRef<IPropertyTypeCustomization> FPropertyEntryCustomization::MakeInstance
 FPropertyEntryCustomization::FPropertyEntryCustomization()
 {
 	PropertyRestriction = MakeShareable(new FPropertyRestriction(FText::FromString("Property already set on for a different entry")));
+	CurrentOptions = nullptr;
 }
 
 void FPropertyEntryCustomization::CustomizeHeader(TSharedRef<IPropertyHandle> PropertyHandle, FDetailWidgetRow& HeaderRow, IPropertyTypeCustomizationUtils& CustomizationUtils)
@@ -30,10 +34,26 @@ void FPropertyEntryCustomization::CustomizeHeader(TSharedRef<IPropertyHandle> Pr
 		MaterialPropertyHandle->CreatePropertyValueWidget()
 	];
 
+	const TArray<TWeakObjectPtr<UObject>>& SelectedObjects = CustomizationUtils.GetPropertyUtilities()->GetSelectedObjects();
+	// Try and find material options instance in currently edited objects
+	CurrentOptions = Cast<UMaterialOptions>((SelectedObjects.FindByPredicate([](TWeakObjectPtr<UObject> Object) { return Cast<UMaterialOptions>(Object.Get()); }))->Get());
+	
 	const int32 Index = PropertyHandle->GetIndexInArray();
-	// Add restriction to ensure the user cannot set up two entries with the same EMaterialProperty value
+
+	// Add restriction to ensure the user cannot set up two entries with the same EMaterialProperty value	
 	MaterialPropertyHandle->AddRestriction(PropertyRestriction.ToSharedRef());
-	MaterialPropertyHandle->SetOnPropertyValueChanged(FSimpleDelegate::CreateLambda([this, Index]() { UpdateRestrictions(Index); }));
+
+	// Set Parent handle on property change to update restrictions across all the entry properties
+	TSharedPtr<IPropertyHandle> ParentHandle = PropertyHandle->GetParentHandle();
+	if (ParentHandle.IsValid())
+	{
+		TSharedPtr<IPropertyHandle> TopParentHandle = ParentHandle->GetParentHandle();
+		if (TopParentHandle.IsValid())
+		{
+			TopParentHandle->SetOnChildPropertyValueChanged(FSimpleDelegate::CreateSP(this, &FPropertyEntryCustomization::UpdateRestrictions, Index));
+		}
+	}
+
 	UpdateRestrictions(Index);
 }
 
@@ -50,17 +70,18 @@ void FPropertyEntryCustomization::CustomizeChildren(TSharedRef<IPropertyHandle>
 
 void FPropertyEntryCustomization::UpdateRestrictions(const int32 EntryIndex)
 {
-	UMaterialOptions* Options = GetMutableDefault<UMaterialOptions>();
 	PropertyRestriction->RemoveAll();
-
-	const UEnum* PropertyEnum = FindObject<UEnum>(ANY_PACKAGE, TEXT("EMaterialProperty"));
-	// Add all previously set material properties to be disabled
-	for (int32 Index = 0; Index < Options->Properties.Num(); ++Index)
+	if (CurrentOptions)
 	{
-		if (Index != EntryIndex)
+		const UEnum* PropertyEnum = FindObject<UEnum>(ANY_PACKAGE, TEXT("EMaterialProperty"));
+		// Add all previously set material properties to be disabled
+		for (int32 Index = 0; Index < CurrentOptions->Properties.Num(); ++Index)
 		{
-			const FPropertyEntry& Entry = Options->Properties[Index];
-			PropertyRestriction->AddDisabledValue(PropertyEnum->GetNameStringByValue(Entry.Property));
+			if (Index != EntryIndex)
+			{
+				const FPropertyEntry& Entry = CurrentOptions->Properties[Index];
+				PropertyRestriction->AddDisabledValue(PropertyEnum->GetNameStringByValue(Entry.Property));
+			}
 		}
 	}
 }
@@ -105,6 +126,31 @@ void FMaterialOptionsCustomization::CustomizeDetails(IDetailLayoutBuilder& Detai
 		AddTextureSizeClamping(TextureSizePropertyHandle);
 	}
 
+	TSharedRef<IPropertyHandle> PropertiesHandle = DetailBuilder.GetProperty(GET_MEMBER_NAME_CHECKED(UMaterialOptions, Properties));
+	if (PropertiesHandle->IsValidHandle())
+	{
+		// Setup delegate for when the number of Material Property items changes
+		FSimpleDelegate RefreshDelegate = FSimpleDelegate::CreateLambda([this, &DetailBuilder, CurrentOptions]()
+		{
+			TArray<EMaterialProperty> Properties;
+			// Ensure that we set duplicate entries to MP_MAX
+			for (FPropertyEntry& Entry : CurrentOptions->Properties)
+			{
+				if (Properties.Contains(Entry.Property))
+				{
+					Entry.Property = MP_MAX;
+				}
+				else if (Entry.Property != MP_MAX)
+				{
+					Properties.Add(Entry.Property);
+				}
+			}
+			
+			DetailBuilder.ForceRefreshDetails();
+		});
+		PropertiesHandle->SetOnPropertyValueChanged(RefreshDelegate);
+	}
+
 	// Only allow changes to LOD indices if we have a valid options instance and if there is actually more than one index
 	ContentBox->SetEnabled(TAttribute<bool>::Create([=]() -> bool { return NumLODs > 1 && CurrentOptions != nullptr; }));
 	for (int32 LODIndex = 0; LODIndex < NumLODs; ++LODIndex)
diff --git a/Engine/Source/Developer/MaterialBaking/Private/MaterialRenderItem.cpp b/Engine/Source/Developer/MaterialBaking/Private/MaterialRenderItem.cpp
--- a/Engine/Source/Developer/MaterialBaking/Private/MaterialRenderItem.cpp
+++ b/Engine/Source/Developer/MaterialBaking/Private/MaterialRenderItem.cpp
@@ -18,34 +18,7 @@ FMeshMaterialRenderItem::FMeshMaterialRenderItem(const FMaterialData* InMaterial
 
 bool FMeshMaterialRenderItem::Render_RenderThread(FRHICommandListImmediate& RHICmdList, FDrawingPolicyRenderState& DrawRenderState, const FCanvas* Canvas)
 {
-	checkSlow(ViewFamily && MaterialSettings && MeshSettings && MaterialRenderProxy);
-	// current render target set for the canvas
-	const FRenderTarget* CanvasRenderTarget = Canvas->GetRenderTarget();
-	const FIntRect ViewRect(FIntPoint(0, 0), CanvasRenderTarget->GetSizeXY());
-
-	// make a temporary view
-	FSceneViewInitOptions ViewInitOptions;
-	ViewInitOptions.ViewFamily = ViewFamily;
-	ViewInitOptions.SetViewRectangle(ViewRect);
-	ViewInitOptions.ViewOrigin = FVector::ZeroVector;
-	ViewInitOptions.ViewRotationMatrix = FMatrix::Identity;
-	ViewInitOptions.ProjectionMatrix = Canvas->GetTransformStack().Top().GetMatrix();
-	ViewInitOptions.BackgroundColor = FLinearColor::Black;
-	ViewInitOptions.OverlayColor = FLinearColor::White;
-
-	const bool bNeedsToSwitchVerticalAxis = RHINeedsToSwitchVerticalAxis(Canvas->GetShaderPlatform()) && !Canvas->GetAllowSwitchVerticalAxis();
-	check(bNeedsToSwitchVerticalAxis == false);
-
-	const FSceneView* View = new FSceneView(ViewInitOptions);
-
-	DrawRenderState.SetBlendState(TStaticBlendState<CW_RGBA>::GetRHI());
-	DrawRenderState.SetDepthStencilState(TStaticDepthStencilState<false, CF_Always>::GetRHI());
-
-	QueueMaterial(RHICmdList, DrawRenderState, View);
-
-	delete View;
-
-	return true;
+	return false;
 }
 
 bool FMeshMaterialRenderItem::Render_GameThread(const FCanvas* Canvas)
diff --git a/Engine/Source/Developer/MaterialBaking/Public/MaterialOptions.h b/Engine/Source/Developer/MaterialBaking/Public/MaterialOptions.h
--- a/Engine/Source/Developer/MaterialBaking/Public/MaterialOptions.h
+++ b/Engine/Source/Developer/MaterialBaking/Public/MaterialOptions.h
@@ -22,7 +22,10 @@ struct FPropertyEntry
 {
 	GENERATED_BODY()
 
-	FPropertyEntry() {}
+	FPropertyEntry() : Property(MP_EmissiveColor), 
+		CustomSize(0, 0),
+		bUseConstantValue(false),
+		ConstantValue(0.0f) {}
 
 	FPropertyEntry(EMaterialProperty InProperty)
 		: Property(InProperty),
diff --git a/Engine/Source/Developer/MeshBoneReduction/Private/MeshBoneReduction.cpp b/Engine/Source/Developer/MeshBoneReduction/Private/MeshBoneReduction.cpp
--- a/Engine/Source/Developer/MeshBoneReduction/Private/MeshBoneReduction.cpp
+++ b/Engine/Source/Developer/MeshBoneReduction/Private/MeshBoneReduction.cpp
@@ -282,19 +282,30 @@ public:
 			TArray<FTransform> MultipliedRefBonePoses;
 			MultipliedRefBonePoses.AddDefaulted(RefBonePoses.Num());
 
+			const bool bDebugBonePoses = false;
+
 			TArray<int32> Processed;
 			Processed.SetNumZeroed(RefBonePoses.Num());
 			// Multiply out parent to child transforms for ref-pose
-			for (int32 BonePoseIndex = 1; BonePoseIndex < RefBonePoses.Num(); BonePoseIndex++)
+			for (int32 BonePoseIndex = 0; BonePoseIndex < RefBonePoses.Num(); BonePoseIndex++)
 			{
 				const int32 BoneIndex = SkeletalMesh->RefSkeleton.FindRawBoneIndex(BoneNames[BonePoseIndex]);
 				const int32 ParentIndex = SkeletalMesh->RefSkeleton.GetParentIndex(BoneIndex);
-				check(ParentIndex == 0 || Processed[ParentIndex] == 1);
-				MultipliedRefBonePoses[BoneIndex] = RefBonePoses[BoneIndex] * MultipliedRefBonePoses[ParentIndex];
-				MultipliedRefBonePoses[BoneIndex].NormalizeRotation();
+				MultipliedRefBonePoses[BoneIndex] = RefBonePoses[BoneIndex];
+
+				if (ParentIndex != INDEX_NONE)
+				{
+					checkf(ParentIndex == 0 || Processed[ParentIndex] == 1, TEXT("Parent bone was not yet processed"));
+					UE_CLOG(bDebugBonePoses, LogMeshBoneReduction, Log, TEXT("Original: [%i]\n%s"), BonePoseIndex, *RefBonePoses[BoneIndex].ToHumanReadableString());
 
-				checkSlow(MultipliedRefBonePoses[BoneIndex].IsRotationNormalized());
-				checkSlow(!MultipliedRefBonePoses[BoneIndex].ContainsNaN());
+					MultipliedRefBonePoses[BoneIndex] = MultipliedRefBonePoses[BoneIndex] * MultipliedRefBonePoses[ParentIndex];
+					MultipliedRefBonePoses[BoneIndex].NormalizeRotation();
+
+					UE_CLOG(bDebugBonePoses, LogMeshBoneReduction, Log, TEXT("Relative: [%i]\n%s"), BonePoseIndex, *MultipliedRefBonePoses[BoneIndex].ToHumanReadableString());
+					checkSlow(MultipliedRefBonePoses[BoneIndex].IsRotationNormalized());
+					checkSlow(!MultipliedRefBonePoses[BoneIndex].ContainsNaN());
+				}
+				
 
 				Processed[BoneIndex] = 1;
 			}
@@ -302,17 +313,24 @@ public:
 			Processed.Empty();
 			Processed.SetNumZeroed(BonePoses.Num());
 			// Multiply out parent to child transforms for bake-pose
-			for (int32 BonePoseIndex = 1; BonePoseIndex < BonePoses.Num(); BonePoseIndex++)
+			for (int32 BonePoseIndex = 0; BonePoseIndex < BonePoses.Num(); BonePoseIndex++)
 			{
 				const int32 BoneIndex = SkeletalMesh->RefSkeleton.FindRawBoneIndex(BoneNames[BonePoseIndex]);
 				const int32 ParentIndex = SkeletalMesh->RefSkeleton.GetParentIndex(BoneIndex);
-				check(ParentIndex == 0 || Processed[ParentIndex] == 1);
-				MultipliedBonePoses[BoneIndex] = BonePoses[BoneIndex] * MultipliedBonePoses[ParentIndex];
+				MultipliedBonePoses[BoneIndex] = BonePoses[BoneIndex];
+				if (ParentIndex != INDEX_NONE)
+				{
+					checkf(ParentIndex == 0 || Processed[ParentIndex] == 1, TEXT("Parent bone was not yet processed"));
+					UE_CLOG(bDebugBonePoses, LogMeshBoneReduction, Log, TEXT("Original: [%i]\n%s"), BonePoseIndex, *BonePoses[BoneIndex].ToHumanReadableString());
 
-				MultipliedBonePoses[BoneIndex].NormalizeRotation();
+					MultipliedBonePoses[BoneIndex] = MultipliedBonePoses[BoneIndex] * MultipliedBonePoses[ParentIndex];
+					MultipliedBonePoses[BoneIndex].NormalizeRotation();
+
+					UE_CLOG(bDebugBonePoses, LogMeshBoneReduction, Log, TEXT("Relative: [%i]\n%s"), BonePoseIndex, *MultipliedBonePoses[BoneIndex].ToHumanReadableString());
+					checkSlow(MultipliedBonePoses[BoneIndex].IsRotationNormalized());
+					checkSlow(!MultipliedBonePoses[BoneIndex].ContainsNaN());
+				}
 
-				checkSlow(MultipliedBonePoses[BoneIndex].IsRotationNormalized());
-				checkSlow(!MultipliedBonePoses[BoneIndex].ContainsNaN());
 				Processed[BoneIndex] = 1;
 			}
 
@@ -320,6 +338,7 @@ public:
 			for (int32 BoneIndex = 0; BoneIndex < NumBones; ++BoneIndex)
 			{
 				MultipliedBonePoses[BoneIndex] = MultipliedRefBonePoses[BoneIndex].Inverse() * MultipliedBonePoses[BoneIndex];
+				UE_CLOG(bDebugBonePoses, LogMeshBoneReduction, Log, TEXT("Final: [%i]\n%s"), BoneIndex, *MultipliedBonePoses[BoneIndex].ToHumanReadableString());
 			}
 		}
 		else
