Move loading TAN depts to log instead of error in phonon.

diff --git a/Engine/Plugins/Runtime/Steam/SteamAudio/Source/SteamAudio/Private/PhononCommon.cpp b/Engine/Plugins/Runtime/Steam/SteamAudio/Source/SteamAudio/Private/PhononCommon.cpp
--- a/Engine/Plugins/Runtime/Steam/SteamAudio/Source/SteamAudio/Private/PhononCommon.cpp
+++ b/Engine/Plugins/Runtime/Steam/SteamAudio/Source/SteamAudio/Private/PhononCommon.cpp
@@ -212,9 +212,9 @@ namespace SteamAudio
 		return StrippedMapName;
 	}
 
-	void* LoadDll(const FString& DllFile)
+	void* LoadDll(const FString& DllFile, bool bErrorOnLoadFailure)
 	{
-		UE_LOG(LogSteamAudio, Log, TEXT("Attempting to load %s"), *DllFile);
+		UE_LOG(LogSteamAudio, Display, TEXT("Attempting to load %s"), *DllFile);
 
 		void* DllHandle = nullptr;
 
@@ -224,16 +224,30 @@ namespace SteamAudio
 		}
 		else
 		{
-			UE_LOG(LogSteamAudio, Error, TEXT("File does not exist. %s"), *DllFile);
+			if (bErrorOnLoadFailure)
+			{
+				UE_LOG(LogSteamAudio, Error, TEXT("Failed to load missing file. %s"), *DllFile);
+			}
+			else
+			{
+				UE_LOG(LogSteamAudio, Display, TEXT("File does not exist. %s"), *DllFile);
+			}
 		}
 
 		if (!DllHandle)
 		{
-			UE_LOG(LogSteamAudio, Error, TEXT("Unable to load %s."), *DllFile);
+			if (bErrorOnLoadFailure)
+			{
+				UE_LOG(LogSteamAudio, Error, TEXT("Failure to load %s."), *DllFile)
+			}
+			else
+			{
+				UE_LOG(LogSteamAudio, Display, TEXT("Unable to load %s."), *DllFile);
+			}
 		}
 		else
 		{
-			UE_LOG(LogSteamAudio, Log, TEXT("Loaded %s."), *DllFile);
+			UE_LOG(LogSteamAudio, Display, TEXT("Loaded %s."), *DllFile);
 		}
 
 		return DllHandle;
diff --git a/Engine/Plugins/Runtime/Steam/SteamAudio/Source/SteamAudio/Private/PhononCommon.h b/Engine/Plugins/Runtime/Steam/SteamAudio/Source/SteamAudio/Private/PhononCommon.h
--- a/Engine/Plugins/Runtime/Steam/SteamAudio/Source/SteamAudio/Private/PhononCommon.h
+++ b/Engine/Plugins/Runtime/Steam/SteamAudio/Source/SteamAudio/Private/PhononCommon.h
@@ -133,8 +133,8 @@ namespace SteamAudio
 	/** Strips PIE prefix from map name for use in-editor. */
 	FString STEAMAUDIO_API StrippedMapName(const FString& MapName);
 
-	/** Attempts to loads the specified DLL, performing some basic error checking. Returns handle to DLL or nullptr on error.*/
-	void* LoadDll(const FString& DllFile);
+	/** Attempts to loads the specified DLL, performing some basic error checking. Returns handle to DLL or nullptr if failed to load.*/
+	void* LoadDll(const FString& DllFile, bool bErrorOnLoadFailure);
 
 	/** Error logs non-successful statuses. */
 	void LogSteamAudioStatus(const IPLerror Status);
