2/2: Missing file
Move loading TAN depts to log instead of error in phonon.

diff --git a/Engine/Plugins/Runtime/Steam/SteamAudio/Source/SteamAudio/Private/SteamAudioModule.cpp b/Engine/Plugins/Runtime/Steam/SteamAudio/Source/SteamAudio/Private/SteamAudioModule.cpp
--- a/Engine/Plugins/Runtime/Steam/SteamAudio/Source/SteamAudio/Private/SteamAudioModule.cpp
+++ b/Engine/Plugins/Runtime/Steam/SteamAudio/Source/SteamAudio/Private/SteamAudioModule.cpp
@@ -50,12 +50,12 @@ namespace SteamAudio
 			FString PathToDll = FPaths::EngineDir() / TEXT("Binaries/ThirdParty/Phonon/Win32/");
 #else
 			FString PathToDll = FPaths::EngineDir() / TEXT("Binaries/ThirdParty/Phonon/Win64/");
-			FSteamAudioModule::TANDllHandle = LoadDll(PathToDll + TEXT("tanrt64.dll"));
-			FSteamAudioModule::TANUtilsDllHandle = LoadDll(PathToDll + TEXT("GPUUtilities.dll"));
+			FSteamAudioModule::TANDllHandle = LoadDll(PathToDll + TEXT("tanrt64.dll"), false);
+			FSteamAudioModule::TANUtilsDllHandle = LoadDll(PathToDll + TEXT("GPUUtilities.dll"), false);
 #endif
 
 			FString DLLToLoad = PathToDll + TEXT("phonon.dll");
-			FSteamAudioModule::PhononDllHandle = LoadDll(DLLToLoad);
+			FSteamAudioModule::PhononDllHandle = LoadDll(DLLToLoad, true);
 		}
 #endif
 		iplCreateContext(UnrealLog, nullptr, nullptr, &GlobalContext);
