Steam Audio dll Failtures in SpatialAudioTemple: Unable to load tantr64.dll and GPUUtilities.dll

diff --git a/Engine/Plugins/Runtime/Steam/SteamAudio/Source/SteamAudio/SteamAudio.Build.cs b/Engine/Plugins/Runtime/Steam/SteamAudio/Source/SteamAudio/SteamAudio.Build.cs
--- a/Engine/Plugins/Runtime/Steam/SteamAudio/Source/SteamAudio/SteamAudio.Build.cs
+++ b/Engine/Plugins/Runtime/Steam/SteamAudio/Source/SteamAudio/SteamAudio.Build.cs
@@ -10,9 +10,6 @@ namespace UnrealBuildTool.Rules
 	{
 		public SteamAudio(ReadOnlyTargetRules Target) : base(Target)
 		{
-
-			OptimizeCode = CodeOptimization.Never;
-
 			PrivateIncludePaths.AddRange(
 				new string[] {
 					"SteamAudio/Private",
@@ -55,13 +52,17 @@ namespace UnrealBuildTool.Rules
 			}
 
 			AddEngineThirdPartyPrivateStaticDependencies(Target, "libPhonon");
-
 			switch (Target.Platform)
 			{
 				case UnrealTargetPlatform.Win32:
+					PrivateDependencyModuleNames.Add("XAudio2");
+					AddEngineThirdPartyPrivateStaticDependencies(Target, "DX11Audio");
+					break;
 				case UnrealTargetPlatform.Win64:
 					PrivateDependencyModuleNames.Add("XAudio2");
 					AddEngineThirdPartyPrivateStaticDependencies(Target, "DX11Audio");
+					PublicDelayLoadDLLs.Add("GPUUtilities.dll");
+					PublicDelayLoadDLLs.Add("tanrt64.dll");
 					break;
 				case UnrealTargetPlatform.Android:
 					string PluginPath = Utils.MakePathRelativeTo(ModuleDirectory, Target.RelativeEnginePath);
