PR #5428: Steam Audio Fixes (Contributed by [REDACTED_PERSON])
* Fix crash when releasing reverb sources.
* Fix random number bug that causes indefinite hang when baking.

diff --git a/Engine/Plugins/Runtime/Steam/SteamAudio/Source/SteamAudio/Private/PhononReverb.cpp b/Engine/Plugins/Runtime/Steam/SteamAudio/Source/SteamAudio/Private/PhononReverb.cpp
--- a/Engine/Plugins/Runtime/Steam/SteamAudio/Source/SteamAudio/Private/PhononReverb.cpp
+++ b/Engine/Plugins/Runtime/Steam/SteamAudio/Source/SteamAudio/Private/PhononReverb.cpp
@@ -80,6 +80,22 @@ namespace SteamAudio
 	void FPhononReverb::Initialize(const FAudioPluginInitializationParams InitializationParams)
 	{
 		AudioPluginInitializationParams = InitializationParams;
+
+		InputAudioFormat.channelLayout = IPL_CHANNELLAYOUT_MONO;
+		InputAudioFormat.channelLayoutType = IPL_CHANNELLAYOUTTYPE_SPEAKERS;
+		InputAudioFormat.channelOrder = IPL_CHANNELORDER_INTERLEAVED;
+		InputAudioFormat.numSpeakers = 1;
+		InputAudioFormat.speakerDirections = nullptr;
+		InputAudioFormat.ambisonicsOrder = -1;
+		InputAudioFormat.ambisonicsNormalization = IPL_AMBISONICSNORMALIZATION_N3D;
+		InputAudioFormat.ambisonicsOrdering = IPL_AMBISONICSORDERING_ACN;
+
+		ReverbSources.SetNum(AudioPluginInitializationParams.NumSources);
+		for (FReverbSource& ReverbSource : ReverbSources)
+		{
+			ReverbSource.InBuffer.format = InputAudioFormat;
+			ReverbSource.InBuffer.numSamples = AudioPluginInitializationParams.BufferLength;
+		}
 	}
 
 	void FPhononReverb::SetEnvironment(FEnvironment* InEnvironment)
@@ -94,15 +110,6 @@ namespace SteamAudio
 		const int32 IndirectImpulseResponseOrder = Environment->GetSimulationSettings().ambisonicsOrder;
 		AmbisonicsChannels = (IndirectImpulseResponseOrder + 1) * (IndirectImpulseResponseOrder + 1);
 
-		InputAudioFormat.channelLayout = IPL_CHANNELLAYOUT_MONO;
-		InputAudioFormat.channelLayoutType = IPL_CHANNELLAYOUTTYPE_SPEAKERS;
-		InputAudioFormat.channelOrder = IPL_CHANNELORDER_INTERLEAVED;
-		InputAudioFormat.numSpeakers = 1;
-		InputAudioFormat.speakerDirections = nullptr;
-		InputAudioFormat.ambisonicsOrder = -1;
-		InputAudioFormat.ambisonicsNormalization = IPL_AMBISONICSNORMALIZATION_N3D;
-		InputAudioFormat.ambisonicsOrdering = IPL_AMBISONICSORDERING_ACN;
-
 		ReverbInputAudioFormat.channelLayout = IPL_CHANNELLAYOUT_STEREO;
 		ReverbInputAudioFormat.channelLayoutType = IPL_CHANNELLAYOUTTYPE_SPEAKERS;
 		ReverbInputAudioFormat.channelOrder = IPL_CHANNELORDER_INTERLEAVED;
@@ -168,13 +175,6 @@ namespace SteamAudio
 		IndirectOutBuffer.interleavedBuffer = IndirectOutArray.GetData();
 		IndirectOutBuffer.deinterleavedBuffer = nullptr;
 
-		ReverbSources.SetNum(AudioPluginInitializationParams.NumSources);
-		for (FReverbSource& ReverbSource : ReverbSources)
-		{
-			ReverbSource.InBuffer.format = InputAudioFormat;
-			ReverbSource.InBuffer.numSamples = AudioPluginInitializationParams.BufferLength;
-		}
-
 		ReverbIndirectContribution = 1.0f;
 
 		CachedSpatializationMethod = GetDefault<USteamAudioSettings>()->IndirectSpatializationMethod;
@@ -232,8 +232,12 @@ namespace SteamAudio
 	{
 		UE_LOG(LogSteamAudio, Log, TEXT("Destroying reverb effect."));
 
-		ReverbSources[SourceId].IndirectContribution = 1.0f;
-		iplDestroyConvolutionEffect(&ReverbSources[SourceId].ConvolutionEffect);
+		check((int32)SourceId < ReverbSources.Num());
+
+		if (ReverbSources[SourceId].ConvolutionEffect)
+		{
+			iplDestroyConvolutionEffect(&ReverbSources[SourceId].ConvolutionEffect);
+		}
 	}
 
 	void FPhononReverb::ProcessSourceAudio(const FAudioPluginSourceInputData& InputData, FAudioPluginSourceOutputData& OutputData)
diff --git a/Engine/Plugins/Runtime/Steam/SteamAudio/Source/SteamAudio/Private/SteamAudioEnvironment.cpp b/Engine/Plugins/Runtime/Steam/SteamAudio/Source/SteamAudio/Private/SteamAudioEnvironment.cpp
--- a/Engine/Plugins/Runtime/Steam/SteamAudio/Source/SteamAudio/Private/SteamAudioEnvironment.cpp
+++ b/Engine/Plugins/Runtime/Steam/SteamAudio/Source/SteamAudio/Private/SteamAudioEnvironment.cpp
@@ -233,11 +233,11 @@ namespace SteamAudio
 
 	int32 FIdentifierMap::Add(const FString& Identifier)
 	{
-		int32 Value = FMath::RandRange(0, TNumericLimits<int32>::Max());
+		int32 Value = FMath::Rand();
 
 		while (ContainsValue(Value))
 		{
-			Value = FMath::RandRange(0, TNumericLimits<int32>::Max());
+			Value = FMath::Rand();
 		}
 
 		IdentifierMapping.Add(TPair<FString, int32>(Identifier, Value));
