Add submenus for Phonon and ITD Spatializer plugins

diff --git a/Engine/Plugins/Runtime/Spatialization/Source/SpatializationEditor/Private/ITDSpatializationSourceSettingsFactory.cpp b/Engine/Plugins/Runtime/Spatialization/Source/SpatializationEditor/Private/ITDSpatializationSourceSettingsFactory.cpp
--- a/Engine/Plugins/Runtime/Spatialization/Source/SpatializationEditor/Private/ITDSpatializationSourceSettingsFactory.cpp
+++ b/Engine/Plugins/Runtime/Spatialization/Source/SpatializationEditor/Private/ITDSpatializationSourceSettingsFactory.cpp
@@ -23,6 +23,15 @@ uint32 FAssetTypeActions_ITDSpatializationSettings::GetCategories()
 	return EAssetTypeCategories::Sounds;
 }
 
+const TArray<FText>& FAssetTypeActions_ITDSpatializationSettings::GetSubMenus() const
+{
+	static const TArray<FText> ITDSubMenus
+	{
+		NSLOCTEXT("Spatialization", "AssetSpatializationSettingsSubMenu", "Spatialization")
+	};
+	return ITDSubMenus;
+}
+
 UITDSpatializationSettingsFactory::UITDSpatializationSettingsFactory(const FObjectInitializer& ObjectInitializer)
 	: Super(ObjectInitializer)
 {
@@ -42,4 +51,4 @@ UObject* UITDSpatializationSettingsFactory::FactoryCreateNew(UClass* Class, UObj
 uint32 UITDSpatializationSettingsFactory::GetMenuCategories() const
 {
 	return EAssetTypeCategories::Sounds;
-}
\ No newline at end of file
+}
diff --git a/Engine/Plugins/Runtime/Spatialization/Source/SpatializationEditor/Private/SpatializationEditorModule.cpp b/Engine/Plugins/Runtime/Spatialization/Source/SpatializationEditor/Private/SpatializationEditorModule.cpp
--- a/Engine/Plugins/Runtime/Spatialization/Source/SpatializationEditor/Private/SpatializationEditorModule.cpp
+++ b/Engine/Plugins/Runtime/Spatialization/Source/SpatializationEditor/Private/SpatializationEditorModule.cpp
@@ -2,15 +2,33 @@
 
 #include "SpatializationEditorModule.h"
 #include "Features/IModularFeatures.h"
+#include "ITDSpatializationSourceSettingsFactory.h"
+
+
+namespace
+{
+	static const FName AssetToolsName = TEXT("AssetTools");
+
+	template <typename T>
+	void AddAssetAction(IAssetTools& AssetTools, TArray<TSharedPtr<FAssetTypeActions_Base>>& AssetActions)
+	{
+		TSharedPtr<T> AssetAction = MakeShared<T>();
+		TSharedPtr<FAssetTypeActions_Base> AssetActionBase = StaticCastSharedPtr<FAssetTypeActions_Base>(AssetAction);
+		AssetTools.RegisterAssetTypeActions(AssetAction.ToSharedRef());
+		AssetActions.Add(AssetActionBase);
+	}
+} // namespace <>
 
 
 void FSpatializationEditorModule::StartupModule()
 {
+	IAssetTools& AssetTools = FModuleManager::LoadModuleChecked<FAssetToolsModule>(AssetToolsName).Get();
+	AddAssetAction<FAssetTypeActions_ITDSpatializationSettings>(AssetTools, AssetActions);
 }
 
 void FSpatializationEditorModule::ShutdownModule()
 {
-	// Nothing done here for now.
+	AssetActions.Reset();
 }
 
 IMPLEMENT_MODULE(FSpatializationEditorModule, SpatializationEditor)
\ No newline at end of file
diff --git a/Engine/Plugins/Runtime/Spatialization/Source/SpatializationEditor/Public/ITDSpatializationSourceSettingsFactory.h b/Engine/Plugins/Runtime/Spatialization/Source/SpatializationEditor/Public/ITDSpatializationSourceSettingsFactory.h
--- a/Engine/Plugins/Runtime/Spatialization/Source/SpatializationEditor/Public/ITDSpatializationSourceSettingsFactory.h
+++ b/Engine/Plugins/Runtime/Spatialization/Source/SpatializationEditor/Public/ITDSpatializationSourceSettingsFactory.h
@@ -16,6 +16,8 @@ public:
 	virtual FColor GetTypeColor() const override;
 	virtual UClass* GetSupportedClass() const override;
 	virtual uint32 GetCategories() override;
+
+	virtual const TArray<FText>& GetSubMenus() const override;
 };
 
 UCLASS(MinimalAPI, hidecategories = Object)
diff --git a/Engine/Plugins/Runtime/Spatialization/Source/SpatializationEditor/Public/SpatializationEditorModule.h b/Engine/Plugins/Runtime/Spatialization/Source/SpatializationEditor/Public/SpatializationEditorModule.h
--- a/Engine/Plugins/Runtime/Spatialization/Source/SpatializationEditor/Public/SpatializationEditorModule.h
+++ b/Engine/Plugins/Runtime/Spatialization/Source/SpatializationEditor/Public/SpatializationEditorModule.h
@@ -1,9 +1,9 @@
 // Copyright 1998-2019 Epic Games, Inc. All Rights Reserved.
-
 #pragma once
 
-
 #include "CoreMinimal.h"
+
+#include "AssetTypeActions_Base.h"
 #include "Modules/ModuleManager.h"
 #include "ITDSpatializer.h"
 
@@ -13,4 +13,7 @@ public:
 	virtual void StartupModule() override;
 
 	virtual void ShutdownModule() override;
+
+private:
+	TArray<TSharedPtr<FAssetTypeActions_Base>> AssetActions;
 };
diff --git a/Engine/Plugins/Runtime/Steam/SteamAudio/Source/SteamAudioEditor/Private/PhononOcclusionSettingsFactory.cpp b/Engine/Plugins/Runtime/Steam/SteamAudio/Source/SteamAudioEditor/Private/PhononOcclusionSettingsFactory.cpp
--- a/Engine/Plugins/Runtime/Steam/SteamAudio/Source/SteamAudioEditor/Private/PhononOcclusionSettingsFactory.cpp
+++ b/Engine/Plugins/Runtime/Steam/SteamAudio/Source/SteamAudioEditor/Private/PhononOcclusionSettingsFactory.cpp
@@ -28,6 +28,15 @@ namespace SteamAudio
 	{
 		return EAssetTypeCategories::Sounds;
 	}
+
+	const TArray<FText>& FAssetTypeActions_PhononOcclusionSettings::GetSubMenus() const
+	{
+		static const TArray<FText> PhononSubMenus
+		{
+			NSLOCTEXT("SteamAudio", "AssetPhononSubMenu", "Phonon")
+		};
+		return PhononSubMenus;
+	}
 }
 
 UPhononOcclusionSettingsFactory::UPhononOcclusionSettingsFactory(const FObjectInitializer& ObjectInitializer)
diff --git a/Engine/Plugins/Runtime/Steam/SteamAudio/Source/SteamAudioEditor/Private/PhononOcclusionSettingsFactory.h b/Engine/Plugins/Runtime/Steam/SteamAudio/Source/SteamAudioEditor/Private/PhononOcclusionSettingsFactory.h
--- a/Engine/Plugins/Runtime/Steam/SteamAudio/Source/SteamAudioEditor/Private/PhononOcclusionSettingsFactory.h
+++ b/Engine/Plugins/Runtime/Steam/SteamAudio/Source/SteamAudioEditor/Private/PhononOcclusionSettingsFactory.h
@@ -19,6 +19,8 @@ namespace SteamAudio
 		virtual FColor GetTypeColor() const override;
 		virtual UClass* GetSupportedClass() const override;
 		virtual uint32 GetCategories() override;
+
+		virtual const TArray<FText>& GetSubMenus() const override;
 	};
 }
 
diff --git a/Engine/Plugins/Runtime/Steam/SteamAudio/Source/SteamAudioEditor/Private/PhononReverbSettingsFactory.cpp b/Engine/Plugins/Runtime/Steam/SteamAudio/Source/SteamAudioEditor/Private/PhononReverbSettingsFactory.cpp
--- a/Engine/Plugins/Runtime/Steam/SteamAudio/Source/SteamAudioEditor/Private/PhononReverbSettingsFactory.cpp
+++ b/Engine/Plugins/Runtime/Steam/SteamAudio/Source/SteamAudioEditor/Private/PhononReverbSettingsFactory.cpp
@@ -28,6 +28,15 @@ namespace SteamAudio
 	{
 		return EAssetTypeCategories::Sounds;
 	}
+
+	const TArray<FText>& FAssetTypeActions_PhononReverbSettings::GetSubMenus() const
+	{
+		static const TArray<FText> PhononSubMenus
+		{
+			NSLOCTEXT("SteamAudio", "AssetPhononSubMenu", "Phonon")
+		};
+		return PhononSubMenus;
+	}
 }
 
 UPhononReverbSettingsFactory::UPhononReverbSettingsFactory(const FObjectInitializer& ObjectInitializer)
diff --git a/Engine/Plugins/Runtime/Steam/SteamAudio/Source/SteamAudioEditor/Private/PhononReverbSettingsFactory.h b/Engine/Plugins/Runtime/Steam/SteamAudio/Source/SteamAudioEditor/Private/PhononReverbSettingsFactory.h
--- a/Engine/Plugins/Runtime/Steam/SteamAudio/Source/SteamAudioEditor/Private/PhononReverbSettingsFactory.h
+++ b/Engine/Plugins/Runtime/Steam/SteamAudio/Source/SteamAudioEditor/Private/PhononReverbSettingsFactory.h
@@ -19,6 +19,8 @@ namespace SteamAudio
 		virtual FColor GetTypeColor() const override;
 		virtual UClass* GetSupportedClass() const override;
 		virtual uint32 GetCategories() override;
+
+		virtual const TArray<FText>& GetSubMenus() const override;
 	};
 }
 
diff --git a/Engine/Plugins/Runtime/Steam/SteamAudio/Source/SteamAudioEditor/Private/PhononSpatializationSettingsFactory.cpp b/Engine/Plugins/Runtime/Steam/SteamAudio/Source/SteamAudioEditor/Private/PhononSpatializationSettingsFactory.cpp
--- a/Engine/Plugins/Runtime/Steam/SteamAudio/Source/SteamAudioEditor/Private/PhononSpatializationSettingsFactory.cpp
+++ b/Engine/Plugins/Runtime/Steam/SteamAudio/Source/SteamAudioEditor/Private/PhononSpatializationSettingsFactory.cpp
@@ -28,6 +28,15 @@ namespace SteamAudio
 	{
 		return EAssetTypeCategories::Sounds;
 	}
+
+	const TArray<FText>& FAssetTypeActions_PhononSpatializationSettings::GetSubMenus() const
+	{
+		static const TArray<FText> PhononSubMenus
+		{
+			NSLOCTEXT("SteamAudio", "AssetPhononSubMenu", "Phonon")
+		};
+		return PhononSubMenus;
+	}
 }
 
 UPhononSpatializationSettingsFactory::UPhononSpatializationSettingsFactory(const FObjectInitializer& ObjectInitializer)
diff --git a/Engine/Plugins/Runtime/Steam/SteamAudio/Source/SteamAudioEditor/Private/PhononSpatializationSettingsFactory.h b/Engine/Plugins/Runtime/Steam/SteamAudio/Source/SteamAudioEditor/Private/PhononSpatializationSettingsFactory.h
--- a/Engine/Plugins/Runtime/Steam/SteamAudio/Source/SteamAudioEditor/Private/PhononSpatializationSettingsFactory.h
+++ b/Engine/Plugins/Runtime/Steam/SteamAudio/Source/SteamAudioEditor/Private/PhononSpatializationSettingsFactory.h
@@ -19,6 +19,8 @@ namespace SteamAudio
 		virtual FColor GetTypeColor() const override;
 		virtual UClass* GetSupportedClass() const override;
 		virtual uint32 GetCategories() override;
+
+		virtual const TArray<FText>& GetSubMenus() const override;
 	};
 }
 
diff --git a/Engine/Plugins/Runtime/Steam/SteamAudio/Source/SteamAudioEditor/Private/SteamAudioEditorModule.cpp b/Engine/Plugins/Runtime/Steam/SteamAudio/Source/SteamAudioEditor/Private/SteamAudioEditorModule.cpp
--- a/Engine/Plugins/Runtime/Steam/SteamAudio/Source/SteamAudioEditor/Private/SteamAudioEditorModule.cpp
+++ b/Engine/Plugins/Runtime/Steam/SteamAudio/Source/SteamAudioEditor/Private/SteamAudioEditorModule.cpp
@@ -23,14 +23,18 @@
 #include "TickableNotification.h"
 #include "PhononScene.h"
 #include "SteamAudioSettings.h"
+#include "IAssetTools.h"
+#include "PhononOcclusionSettingsFactory.h"
 #include "PhononProbeVolume.h"
 #include "PhononProbeVolumeDetails.h"
-#include "PhononScene.h"
 #include "PhononProbeComponent.h"
 #include "PhononProbeComponentVisualizer.h"
+#include "PhononReverbSettingsFactory.h"
+#include "PhononScene.h"
 #include "PhononSourceComponent.h"
 #include "PhononSourceComponentDetails.h"
 #include "PhononSourceComponentVisualizer.h"
+#include "PhononSpatializationSettingsFactory.h"
 #include "PhononCommon.h"
 #include "SteamAudioEdMode.h"
 
@@ -47,6 +51,21 @@ DEFINE_LOG_CATEGORY(LogSteamAudioEditor);
 
 IMPLEMENT_MODULE(SteamAudio::FSteamAudioEditorModule, SteamAudioEditor)
 
+
+namespace
+{
+	static const FName AssetToolsName = TEXT("AssetTools");
+
+	template <typename T>
+	void AddAssetAction(IAssetTools& AssetTools, TArray<TSharedPtr<FAssetTypeActions_Base>>& AssetActions)
+	{
+		TSharedPtr<T> AssetAction = MakeShared<T>();
+		TSharedPtr<FAssetTypeActions_Base> AssetActionBase = StaticCastSharedPtr<FAssetTypeActions_Base>(AssetAction);
+		AssetTools.RegisterAssetTypeActions(AssetAction.ToSharedRef());
+		AssetActions.Add(AssetActionBase);
+	}
+} // namespace <>
+
 namespace SteamAudio
 {
 	//==============================================================================================================================================
@@ -127,6 +146,11 @@ namespace SteamAudio
 		SteamAudioStyleSet->Set("LevelEditor.SteamAudioMode", new FSlateImageBrush(SteamAudioContent + "/SteamAudio_EdMode_40.png", Vec40));
 		SteamAudioStyleSet->Set("LevelEditor.SteamAudioMode.Small", new FSlateImageBrush(SteamAudioContent + "/SteamAudio_EdMode_16.png", Vec16));
 
+		IAssetTools& AssetTools = FModuleManager::LoadModuleChecked<FAssetToolsModule>(AssetToolsName).Get();
+		AddAssetAction<FAssetTypeActions_PhononReverbSettings>(AssetTools, AssetActions);
+		AddAssetAction<FAssetTypeActions_PhononOcclusionSettings>(AssetTools, AssetActions);
+		AddAssetAction<FAssetTypeActions_PhononSpatializationSettings>(AssetTools, AssetActions);
+
 		FSlateStyleRegistry::RegisterSlateStyle(*SteamAudioStyleSet.Get());
 
 		// Register the ed mode
@@ -152,6 +176,8 @@ namespace SteamAudio
 			}
 		}
 
+		AssetActions.Reset();
+
 		FSlateStyleRegistry::UnRegisterSlateStyle(*SteamAudioStyleSet.Get());
 	}
 
diff --git a/Engine/Plugins/Runtime/Steam/SteamAudio/Source/SteamAudioEditor/Private/SteamAudioEditorModule.h b/Engine/Plugins/Runtime/Steam/SteamAudio/Source/SteamAudioEditor/Private/SteamAudioEditorModule.h
--- a/Engine/Plugins/Runtime/Steam/SteamAudio/Source/SteamAudioEditor/Private/SteamAudioEditorModule.h
+++ b/Engine/Plugins/Runtime/Steam/SteamAudio/Source/SteamAudioEditor/Private/SteamAudioEditorModule.h
@@ -5,6 +5,8 @@
 #pragma once
 
 #include "ISteamAudioEditorModule.h"
+
+#include "AssetTypeActions_Base.h"
 #include "Framework/MultiBox/MultiBoxExtender.h"
 #include "PhononScene.h"
 
@@ -39,5 +41,7 @@ namespace SteamAudio
 		TSharedPtr<FBakeIndirectWindow> BakeIndirectWindow;
 		TArray<FName> RegisteredComponentClassNames;
 		FPhononSceneInfo CurrentPhononSceneInfo;
+
+		TArray<TSharedPtr<FAssetTypeActions_Base>> AssetActions;
 	};
 }
