Fixing compiler errors in Oculus/Phonon plugins

diff --git a/Engine/Plugins/Runtime/Oculus/OculusAudio/Source/OculusAudio/Private/OculusAudioReverb.cpp b/Engine/Plugins/Runtime/Oculus/OculusAudio/Source/OculusAudio/Private/OculusAudioReverb.cpp
--- a/Engine/Plugins/Runtime/Oculus/OculusAudio/Source/OculusAudio/Private/OculusAudioReverb.cpp
+++ b/Engine/Plugins/Runtime/Oculus/OculusAudio/Source/OculusAudio/Private/OculusAudioReverb.cpp
@@ -48,13 +48,16 @@ void OculusAudioReverb::SetContext(ovrAudioContext* SharedContext)
 		Context = SharedContext;
 	}
 
-	Submix->SetContext(SharedContext);
+	if (Submix.IsValid())
+	{
+		Submix->SetContext(SharedContext);
+	}
 }
 
 void OculusAudioReverb::ClearContext()
 {
 	Context = nullptr;
-	for (FSubmixEffectOculusReverbPlugin* Submix : Submixes)
+	if (Submix.IsValid())
 	{
 		Submix->ClearContext();
 	}
diff --git a/Engine/Plugins/Runtime/Oculus/OculusAudio/Source/OculusAudio/Private/OculusAudioReverb.h b/Engine/Plugins/Runtime/Oculus/OculusAudio/Source/OculusAudio/Private/OculusAudioReverb.h
--- a/Engine/Plugins/Runtime/Oculus/OculusAudio/Source/OculusAudio/Private/OculusAudioReverb.h
+++ b/Engine/Plugins/Runtime/Oculus/OculusAudio/Source/OculusAudio/Private/OculusAudioReverb.h
@@ -62,7 +62,7 @@ public:
 		return; // PAS
 	}
 
-	virtual TSharedPtr<FSoundEffectSubmix> GetEffectSubmix() override;
+	virtual FSoundEffectSubmixPtr GetEffectSubmix() override;
 
 	virtual USoundSubmix* GetSubmix() override;
 
diff --git a/Engine/Plugins/Runtime/ResonanceAudio/Source/ResonanceAudio/Private/ResonanceAudioReverb.cpp b/Engine/Plugins/Runtime/ResonanceAudio/Source/ResonanceAudio/Private/ResonanceAudioReverb.cpp
--- a/Engine/Plugins/Runtime/ResonanceAudio/Source/ResonanceAudio/Private/ResonanceAudioReverb.cpp
+++ b/Engine/Plugins/Runtime/ResonanceAudio/Source/ResonanceAudio/Private/ResonanceAudioReverb.cpp
@@ -138,8 +138,7 @@ namespace ResonanceAudio
 			ReverbSubmixEffect->SetResonanceAudioReverbPlugin(this);
 
 			ReverbSubmixEffect->SetEnabled(true);
-			SubmixEffect = MakeShareable(ReverbSubmixEffect);
-
+			SubmixEffect = MakeShareable(static_cast<FSoundEffectSubmix*>(ReverbSubmixEffect));
 		}
 		else
 		{
diff --git a/Engine/Plugins/Runtime/Steam/SteamAudio/Source/SteamAudio/Private/PhononReverb.cpp b/Engine/Plugins/Runtime/Steam/SteamAudio/Source/SteamAudio/Private/PhononReverb.cpp
--- a/Engine/Plugins/Runtime/Steam/SteamAudio/Source/SteamAudio/Private/PhononReverb.cpp
+++ b/Engine/Plugins/Runtime/Steam/SteamAudio/Source/SteamAudio/Private/PhononReverb.cpp
@@ -398,10 +398,10 @@ namespace SteamAudio
 			USubmixEffectReverbPluginPreset* ReverbPluginPreset = NewObject<USubmixEffectReverbPluginPreset>(GetSubmix(), TEXT("Master Reverb Plugin Effect Preset"));
 			FSubmixEffectReverbPlugin* ReverbSubmixEffect = static_cast<FSubmixEffectReverbPlugin*>(ReverbPluginPreset->CreateNewEffect());
 			ReverbSubmixEffect->SetPhononReverbPlugin(this);
-			USubmixEffectReverbPluginPreset
+
 			ReverbSubmixEffect->SetEnabled(true);
 
-			SubmixEffect = MakeShareable(static_cast<FSoundEffectSubmix>(ReverbSubmixEffect));
+			SubmixEffect = MakeShareable(static_cast<FSoundEffectSubmix*>(ReverbSubmixEffect));
 		}
 
 		return SubmixEffect;
