Steam Audio Submix Fixes
- Fix bad steam audio default submix path
- Safeguard against bad initialization causing null ref in OnProcessAudio
- Fix initialization by ensuring content is providing default plugin effect/submix and warn adequately if not, while stubbing in correctly
- Ensure Phonon is manipulating (i.e. registering plugin) with the first effect and verifying its type as phonon reverb effect type, allowing multiple effects to inhabit default submix if desired
- Replace default submix effect instead of adding (to avoid having multiple effects being included on default plugin reverb submix but one being unused stub)

#rnx

diff --git a/Engine/Plugins/Runtime/Steam/SteamAudio/Source/SteamAudio/Private/PhononReverb.cpp b/Engine/Plugins/Runtime/Steam/SteamAudio/Source/SteamAudio/Private/PhononReverb.cpp
--- a/Engine/Plugins/Runtime/Steam/SteamAudio/Source/SteamAudio/Private/PhononReverb.cpp
+++ b/Engine/Plugins/Runtime/Steam/SteamAudio/Source/SteamAudio/Private/PhononReverb.cpp
@@ -395,8 +395,23 @@ namespace SteamAudio
 	{
 		if (!SubmixEffect.IsValid())
 		{
-			USubmixEffectReverbPluginPreset* ReverbPluginPreset = NewObject<USubmixEffectReverbPluginPreset>(GetSubmix(), TEXT("Reverb Plugin Effect Preset"));
+			USubmixEffectReverbPluginPreset* ReverbPluginPreset = nullptr;
+			USoundSubmix* Submix = GetSubmix();
+			if (Submix->SubmixEffectChain.Num() > 0)
+			{
+				if (USubmixEffectReverbPluginPreset* CurrentPreset = Cast<USubmixEffectReverbPluginPreset>(Submix->SubmixEffectChain[0]))
+				{
+					ReverbPluginPreset = CurrentPreset;
+				}
+			}
+
+			if (!ReverbPluginPreset)
+			{
+				ReverbPluginPreset = NewObject<USubmixEffectReverbPluginPreset>(Submix, TEXT("Reverb Plugin Effect Preset"));
+			}
+
 			FSubmixEffectReverbPlugin* ReverbSubmixEffect = static_cast<FSubmixEffectReverbPlugin*>(ReverbPluginPreset->CreateNewEffect());
+			ReverbSubmixEffect->SetPreset(ReverbPluginPreset);
 			ReverbSubmixEffect->SetPhononReverbPlugin(this);
 
 			ReverbSubmixEffect->SetEnabled(true);
@@ -411,8 +426,13 @@ namespace SteamAudio
 	{
 		const USteamAudioSettings* Settings = GetDefault<USteamAudioSettings>();
 		check(Settings);
-		USoundSubmix* ReverbSubmix = Cast<USoundSubmix>(Settings->OutputSubmix.TryLoad());
-		if (!ReverbSubmix)
+		
+		if (!ReverbSubmix.IsValid())
+		{
+			ReverbSubmix = Cast<USoundSubmix>(Settings->OutputSubmix.TryLoad());
+		}
+
+		if (!ReverbSubmix.IsValid())
 		{
 			static const FString DefaultSubmixName = TEXT("Phonon Reverb Submix");
 			UE_LOG(LogSteamAudio, Error, TEXT("Failed to load Phonon Reverb Submix from object path '%s' in PhononSettings. Creating '%s' as stub."),
@@ -436,13 +456,13 @@ namespace SteamAudio
 		if (!bFoundPreset)
 		{
 			static const FString DefaultPresetName = TEXT("PhononReverbDefault_0");
-			UE_LOG(LogSteamAudio, Error, TEXT("Failed to find Phonon USubmixEffectReverbPluginPreset on default reverb submix. Creating stub '%s'."),
+			UE_LOG(LogSteamAudio, Error, TEXT("Failed to find Phonon USubmixEffectReverbPluginPreset on default Phonon Reverb Submix. Creating stub '%s'."),
 				*Settings->OutputSubmix.GetAssetPathString(),
 				*DefaultPresetName);
 			ReverbSubmix->SubmixEffectChain.Add(NewObject<USubmixEffectReverbPluginPreset>(USubmixEffectReverbPluginPreset::StaticClass(), *DefaultPresetName));
 		}
 
-		return ReverbSubmix;
+		return ReverbSubmix.Get();
 	}
 
 	//==============================================================================================================================================
@@ -482,7 +502,10 @@ uint32 FSubmixEffectReverbPlugin::GetDesiredInputChannelCountOverride() const
 
 void FSubmixEffectReverbPlugin::OnProcessAudio(const FSoundEffectSubmixInputData& InData, FSoundEffectSubmixOutputData& OutData)
 {
-	PhononReverbPlugin->ProcessMixedAudio(InData, OutData);
+	if (PhononReverbPlugin)
+	{
+		PhononReverbPlugin->ProcessMixedAudio(InData, OutData);
+	}
 }
 
 void FSubmixEffectReverbPlugin::SetPhononReverbPlugin(SteamAudio::FPhononReverb* InPhononReverbPlugin)
diff --git a/Engine/Plugins/Runtime/Steam/SteamAudio/Source/SteamAudio/Private/PhononReverb.h b/Engine/Plugins/Runtime/Steam/SteamAudio/Source/SteamAudio/Private/PhononReverb.h
--- a/Engine/Plugins/Runtime/Steam/SteamAudio/Source/SteamAudio/Private/PhononReverb.h
+++ b/Engine/Plugins/Runtime/Steam/SteamAudio/Source/SteamAudio/Private/PhononReverb.h
@@ -7,6 +7,8 @@
 #include "IAudioExtensionPlugin.h"
 #include "Sound/SoundEffectSubmix.h"
 #include "Sound/SoundEffectPreset.h"
+#include "UObject/WeakObjectPtrTemplates.h"
+
 #include "PhononCommon.h"
 #include "phonon.h"
 #include "PhononReverb.generated.h"
@@ -55,6 +57,7 @@ namespace SteamAudio
 
 	private:
 		FSoundEffectSubmixPtr SubmixEffect;
+		TWeakObjectPtr<USoundSubmix> ReverbSubmix;
 
 		IPLhandle BinauralRenderer;
 		IPLhandle IndirectBinauralEffect;
diff --git a/Engine/Plugins/Runtime/Steam/SteamAudio/Source/SteamAudio/Private/SteamAudioSettings.cpp b/Engine/Plugins/Runtime/Steam/SteamAudio/Source/SteamAudio/Private/SteamAudioSettings.cpp
--- a/Engine/Plugins/Runtime/Steam/SteamAudio/Source/SteamAudio/Private/SteamAudioSettings.cpp
+++ b/Engine/Plugins/Runtime/Steam/SteamAudio/Source/SteamAudio/Private/SteamAudioSettings.cpp
@@ -38,7 +38,7 @@ USteamAudioSettings::USteamAudioSettings()
 	, TANMaxSources(32)
 	, RadeonRaysBakingBatchSize(1)
 {
-	OutputSubmix = FString(TEXT("/PhononAudio/PhononSubmixDefault.PhononSubmixDefault"));
+	OutputSubmix = FString(TEXT("/SteamAudio/PhononSubmixDefault.PhononSubmixDefault"));
 
 	auto MaterialPreset = SteamAudio::MaterialPresets[StaticMeshMaterialPreset];
 	StaticMeshLowFreqAbsorption = MaterialPreset.lowFreqAbsorption;
diff --git a/Engine/Source/Runtime/AudioMixer/Private/AudioMixerDevice.cpp b/Engine/Source/Runtime/AudioMixer/Private/AudioMixerDevice.cpp
--- a/Engine/Source/Runtime/AudioMixer/Private/AudioMixerDevice.cpp
+++ b/Engine/Source/Runtime/AudioMixer/Private/AudioMixerDevice.cpp
@@ -668,13 +668,14 @@ namespace Audio
 					check(ReverbPluginMixerSubmixPtr.IsValid());
 
 					const uint32 ReverbPluginId = Preset->GetUniqueID();
+					const TWeakObjectPtr<USoundSubmix> ReverbPluginSubmixPtr = ReverbPluginSubmix;
 					FMixerSubmixWeakPtr ReverbPluginMixerSubmixWeakPtr = ReverbPluginMixerSubmixPtr;
-					AudioRenderThreadCommand([ReverbPluginMixerSubmixWeakPtr, ReverbPluginEffectSubmix, ReverbPluginId]()
+					AudioRenderThreadCommand([ReverbPluginMixerSubmixWeakPtr, ReverbPluginSubmixPtr, ReverbPluginEffectSubmix, ReverbPluginId]()
 					{
 						FMixerSubmixPtr PluginSubmixPtr = ReverbPluginMixerSubmixWeakPtr.Pin();
-						if (PluginSubmixPtr.IsValid())
+						if (PluginSubmixPtr.IsValid() && ReverbPluginSubmixPtr.IsValid())
 						{
-							PluginSubmixPtr->AddSoundEffectSubmix(ReverbPluginId, ReverbPluginEffectSubmix);
+							PluginSubmixPtr->ReplaceSoundEffectSubmix(0, ReverbPluginId, ReverbPluginEffectSubmix);
 						}
 					});
 				}
diff --git a/Engine/Source/Runtime/AudioMixer/Private/AudioMixerSubmix.cpp b/Engine/Source/Runtime/AudioMixer/Private/AudioMixerSubmix.cpp
--- a/Engine/Source/Runtime/AudioMixer/Private/AudioMixerSubmix.cpp
+++ b/Engine/Source/Runtime/AudioMixer/Private/AudioMixerSubmix.cpp
@@ -108,7 +108,7 @@ namespace Audio
 			{
 				// Create a new effect instance using the preset
 				FSoundEffectSubmix* SubmixEffect = static_cast<FSoundEffectSubmix*>(EffectPreset->CreateNewEffect());
-
+				
 				FSoundEffectSubmixInitData InitData;
 				InitData.SampleRate = MixerDevice->GetSampleRate();
 				InitData.PresetSettings = nullptr;
@@ -376,6 +376,15 @@ namespace Audio
 		EffectSubmixChain.Reset();
 	}
 
+	void FMixerSubmix::ReplaceSoundEffectSubmix(int32 InIndex, int32 InPresetId, FSoundEffectSubmixPtr InEffectInstance)
+	{
+		if (InEffectInstance->GetPreset() != nullptr && InIndex < EffectSubmixChain.Num())
+		{
+			EffectSubmixChain[InIndex].PresetId = InPresetId;
+			EffectSubmixChain[InIndex].EffectInstance = InEffectInstance;
+		}
+	}
+
 	void FMixerSubmix::SetBackgroundMuted(bool bInMuted)
 	{
 		SubmixCommand([this, bInMuted]()
diff --git a/Engine/Source/Runtime/AudioMixer/Public/AudioMixerSubmix.h b/Engine/Source/Runtime/AudioMixer/Public/AudioMixerSubmix.h
--- a/Engine/Source/Runtime/AudioMixer/Public/AudioMixerSubmix.h
+++ b/Engine/Source/Runtime/AudioMixer/Public/AudioMixerSubmix.h
@@ -102,6 +102,9 @@ namespace Audio
 		/** Clears all submix effects from the effect submix chain. */
 		void ClearSoundEffectSubmixes();
 
+		/** Swaps effect for provided submix at the given index.  Fails if effect at index doesn't exist */
+		void ReplaceSoundEffectSubmix(int32 InIndex, int32 InPresetId, FSoundEffectSubmixPtr InEffectInstance);
+
 		/** Whether or not this submix instance is muted. */
 		void SetBackgroundMuted(bool bInMuted);
 
