Unshelved from pending changelist '10947328':

PR #6538: SteamAudio: Add Dynamic Geometry Support (Contributed by [REDACTED_PERSON])

diff --git a/Engine/Plugins/Runtime/Steam/SteamAudio/Source/SteamAudio/Private/PhononCommon.cpp b/Engine/Plugins/Runtime/Steam/SteamAudio/Source/SteamAudio/Private/PhononCommon.cpp
--- a/Engine/Plugins/Runtime/Steam/SteamAudio/Source/SteamAudio/Private/PhononCommon.cpp
+++ b/Engine/Plugins/Runtime/Steam/SteamAudio/Source/SteamAudio/Private/PhononCommon.cpp
@@ -75,7 +75,9 @@ namespace SteamAudio
 
 	FString BasePath;
 	FString RuntimePath;
+	FString DynamicRuntimePath;
 	FString EditorOnlyPath;
+	FString DynamicEditorOnlyPath;
 
 	void* UnrealAlloc(const size_t size, const size_t alignment)
 	{
@@ -169,12 +171,12 @@ namespace SteamAudio
 	 * Given a Unreal transform, convert it to a 4x4 column-major transformation matrix. OutMatrix is assumed to be a contiguous array of
 	 * 16 floats.
 	 */
-	void GetMatrixForTransform(const FTransform& Transform, float* OutMatrix)
+	void GetMatrixForTransform(const FTransform& Transform, float* OutMatrix, bool bOutputRowMajor /* = false */, bool bApplyScale /* = true */)
 	{
 		check(OutMatrix);
 
 		auto PhononTranslation = SteamAudio::UnrealToPhononFVector(Transform.GetTranslation());
-		auto PhononScale = SteamAudio::UnrealToPhononFVector(Transform.GetScale3D());
+		auto PhononScale = SteamAudio::UnrealToPhononFVector(Transform.GetScale3D(), bApplyScale);
 
 		FQuat RotationQuatConverted;
 		RotationQuatConverted.X = -Transform.GetRotation().Y;
@@ -186,16 +188,41 @@ namespace SteamAudio
 		FScaleMatrix ScaleMatrix(PhononScale);
 		auto ConvertedMatrix = (ScaleMatrix * RotationTranslationMatrix).GetTransposed();
 		
-		// Convert row-major to column-major
-		for (int32 i = 0; i < 4; ++i)
+		// Convert to column major if bOutputRowMajor is false, otherwise output column major
+		if (bOutputRowMajor)
 		{
-			for (int32 j = 0; j < 4; ++j)
+			// Row major order (flatten)
+			for (int32 i = 0; i < 4; ++i)
 			{
-				OutMatrix[j * 4 + i] = ConvertedMatrix.M[i][j];
+				for (int32 j = 0; j < 4; ++j)
+				{
+					OutMatrix[i * 4 + j] = ConvertedMatrix.M[i][j];
+				}
+			}		
+		}
+		else
+		{
+			// Column major order (conversion from row major order)
+			for (int32 i = 0; i < 4; ++i)
+			{
+				for (int32 j = 0; j < 4; ++j)
+				{
+					OutMatrix[j * 4 + i] = ConvertedMatrix.M[i][j];
+				}
 			}
 		}
 	}
 
+	IPLMatrix4x4 GetIPLMatrix(float* InMatrix)
+	{
+		check(InMatrix);
+
+		IPLMatrix4x4 PhononMatrix;
+		FMemory::Memcpy(&PhononMatrix, InMatrix, sizeof(PhononMatrix));
+		
+		return PhononMatrix;
+	}
+
 	FString StrippedMapName(const FString& MapName)
 	{
 		const FRegexPattern PieMapPattern(TEXT("UEDPIE_\\d_"));
diff --git a/Engine/Plugins/Runtime/Steam/SteamAudio/Source/SteamAudio/Private/PhononCommon.h b/Engine/Plugins/Runtime/Steam/SteamAudio/Source/SteamAudio/Private/PhononCommon.h
--- a/Engine/Plugins/Runtime/Steam/SteamAudio/Source/SteamAudio/Private/PhononCommon.h
+++ b/Engine/Plugins/Runtime/Steam/SteamAudio/Source/SteamAudio/Private/PhononCommon.h
@@ -8,6 +8,8 @@
 #include "CoreMinimal.h"
 #include "UObject/ObjectMacros.h"
 
+class UPhononGeometryComponent;
+
 DECLARE_LOG_CATEGORY_EXTERN(LogSteamAudio, Log, All);
 
 UENUM(BlueprintType)
@@ -106,6 +108,14 @@ namespace SteamAudio
 		int32 SecondaryRays;
 	};
 
+	struct FDynamicGeometryMap
+	{
+		FString ID;
+		UPhononGeometryComponent* DynamicGeometryComponent;
+		IPLhandle DynamicScene;
+		IPLhandle DynamicGeometry;
+	};
+
 	/** 1 Unreal Unit = 1cm, 1 Phonon Unit = 1m */
 	const float SCALEFACTOR = 0.01f;
 
@@ -120,7 +130,9 @@ namespace SteamAudio
 
 	extern FString STEAMAUDIO_API BasePath;
 	extern FString STEAMAUDIO_API RuntimePath;
+	extern FString STEAMAUDIO_API DynamicRuntimePath;
 	extern FString STEAMAUDIO_API EditorOnlyPath;
+	extern FString STEAMAUDIO_API DynamicEditorOnlyPath;
 
 	/** Functions to convert to/from Phonon/UE4 and IPLVector3/FVector. */
 	IPLVector3 STEAMAUDIO_API IPLVector3FromFVector(const FVector& Coords);
@@ -131,7 +143,10 @@ namespace SteamAudio
 	IPLVector3 STEAMAUDIO_API PhononToUnrealIPLVector3(const FVector& Coords, const bool bScale = true);
 
 	/** Given a UE4 transform, produces a corresponding 4x4 transformation matrix. */
-	void STEAMAUDIO_API GetMatrixForTransform(const FTransform& Transform, float* OutMatrix);
+	void STEAMAUDIO_API GetMatrixForTransform(const FTransform& Transform, float* OutMatrix, bool bOutputRowMajor = false, bool bApplyScale = true);
+
+	/** Given a matrix (16 contiguous floats), return a corresponding IPL 4x4 transformation matrix. */
+	IPLMatrix4x4 STEAMAUDIO_API GetIPLMatrix(float* InMatrix);
 
 	/** Phonon raytracer callback that routes ClosestHit queries to the UE4 raytracer. */
 	void STEAMAUDIO_API ClosestHit(const IPLfloat32* Origin, const IPLfloat32* Direction, const IPLfloat32 MinDistance,
diff --git a/Engine/Plugins/Runtime/Steam/SteamAudio/Source/SteamAudio/Private/PhononGeometryComponent.cpp b/Engine/Plugins/Runtime/Steam/SteamAudio/Source/SteamAudio/Private/PhononGeometryComponent.cpp
--- a/Engine/Plugins/Runtime/Steam/SteamAudio/Source/SteamAudio/Private/PhononGeometryComponent.cpp
+++ b/Engine/Plugins/Runtime/Steam/SteamAudio/Source/SteamAudio/Private/PhononGeometryComponent.cpp
@@ -4,8 +4,8 @@
 
 #include "PhononGeometryComponent.h"
 #include "PhononScene.h"
-#include "Engine/StaticMeshActor.h"
 #include "StaticMeshResources.h"
+#include "PhononPluginManager.h"
 
 UPhononGeometryComponent::UPhononGeometryComponent()
 	: ExportAllChildren(false)
@@ -14,6 +14,65 @@ UPhononGeometryComponent::UPhononGeometryComponent()
 {
 }
 
+
+void UPhononGeometryComponent::BeginPlay()
+{
+	Super::BeginPlay();
+
+	// Check if we are parented to a movable mesh
+	USceneComponent* Parent = GetAttachParent();
+
+	if (Parent
+		&& Parent->IsValidLowLevelFast()
+		&& Parent->Mobility == EComponentMobility::Movable
+		)
+	{
+		bHasDynamicParent = true;
+
+		// Register this component to the plugin manager as holding on to dynamic geometry
+		FAudioDeviceHandle AudioDevice = GEngine->GetMainAudioDevice();
+		if (AudioDevice)
+		{
+			for (TAudioPluginListenerPtr AudioListener : AudioDevice->PluginListeners)
+			{
+				auto SteamAudioListener = StaticCastSharedPtr<SteamAudio::FPhononPluginManager>(AudioListener);
+
+				if (SteamAudioListener && SteamAudioListener->PluginID.IsEqual(FName(TEXT("Phonon"))))
+				{
+					// Add dynamic geometry to runtime cache array
+					SteamAudioListener->AddDynamicGeometry(this);
+				}
+			}
+		}
+	}
+}
+
+
+void UPhononGeometryComponent::EndPlay(const EEndPlayReason::Type EndPlayReason)
+{
+	Super::EndPlay(EndPlayReason);
+
+	// Destroy any instanced meshes created
+	if (bHasDynamicParent)
+	{	
+		// Unregister this component from the plugin manager
+		FAudioDeviceHandle AudioDevice = GEngine->GetMainAudioDevice();
+		if (AudioDevice)
+		{
+			for (TAudioPluginListenerPtr AudioListener : AudioDevice->PluginListeners)
+			{
+				auto SteamAudioListener = StaticCastSharedPtr<SteamAudio::FPhononPluginManager>(AudioListener);
+
+				if (SteamAudioListener && SteamAudioListener->PluginID.IsEqual(FName(TEXT("Phonon"))))
+				{
+					SteamAudioListener->RemoveDynamicGeometry(this);
+				}
+			}
+		}
+	}
+
+}
+
 void UPhononGeometryComponent::UpdateStatistics()
 {
 #if WITH_EDITOR
diff --git a/Engine/Plugins/Runtime/Steam/SteamAudio/Source/SteamAudio/Private/PhononGeometryComponent.h b/Engine/Plugins/Runtime/Steam/SteamAudio/Source/SteamAudio/Private/PhononGeometryComponent.h
--- a/Engine/Plugins/Runtime/Steam/SteamAudio/Source/SteamAudio/Private/PhononGeometryComponent.h
+++ b/Engine/Plugins/Runtime/Steam/SteamAudio/Source/SteamAudio/Private/PhononGeometryComponent.h
@@ -4,7 +4,8 @@
 
 #pragma once
 
-#include "Components/ActorComponent.h"
+#include "Components/SceneComponent.h"
+#include "PhononCommon.h"
 #include "PhononGeometryComponent.generated.h"
 
 /**
@@ -12,13 +13,14 @@
  * Should be placed on Static Mesh actors.
  */
 UCLASS(ClassGroup = (Audio), HideCategories = (Activation, Collision, Cooking), meta = (BlueprintSpawnableComponent))
-class UPhononGeometryComponent : public UActorComponent
+class UPhononGeometryComponent : public USceneComponent
 {
 	GENERATED_BODY()
 
 public:
 	UPhononGeometryComponent();
-
+	virtual void BeginPlay() override;
+	virtual void EndPlay(const EEndPlayReason::Type EndPlayReason) override;
 	void OnComponentCreated();
 
 #if WITH_EDITOR
@@ -37,6 +39,11 @@ public:
 	UPROPERTY(VisibleAnywhere, Category = GeometryStatistics, meta = (DisplayName = "Triangles"))
 	uint32 NumTriangles;
 
+	// The handle to an instanced mesh if this component is attached to dynamic geometry
+	IPLhandle* InstancedMesh = nullptr;
+
 private:
 	void UpdateStatistics();
+	bool bHasDynamicParent = false;
+
 }; 
\ No newline at end of file
diff --git a/Engine/Plugins/Runtime/Steam/SteamAudio/Source/SteamAudio/Private/PhononPluginManager.cpp b/Engine/Plugins/Runtime/Steam/SteamAudio/Source/SteamAudio/Private/PhononPluginManager.cpp
--- a/Engine/Plugins/Runtime/Steam/SteamAudio/Source/SteamAudio/Private/PhononPluginManager.cpp
+++ b/Engine/Plugins/Runtime/Steam/SteamAudio/Source/SteamAudio/Private/PhononPluginManager.cpp
@@ -5,9 +5,16 @@
 #include "PhononPluginManager.h"
 #include "SteamAudioModule.h"
 #include "AudioDevice.h"
+#include "HAL/FileManager.h"
+#include "DrawDebugHelpers.h"
 
 namespace SteamAudio
 {
+	const FName FPhononPluginManager::PluginID = FName("Phonon");
+	TArray<FDynamicGeometryMap> FPhononPluginManager::DynamicGeometry;
+	TMap<FString, IPLhandle> FPhononPluginManager::DynamicSceneMap;
+	UWorld* FPhononPluginManager::CurrentLevel;
+		
 	FPhononPluginManager::FPhononPluginManager()
 		: bEnvironmentInitialized(false)
 		, ReverbPtr(nullptr)
@@ -29,20 +36,26 @@ namespace SteamAudio
 	{
 		if (ListenerWorld->WorldType == EWorldType::Editor || bEnvironmentInitialized)
 		{
-			UE_LOG(LogSteamAudio, Log, TEXT("Trying to initialize environment with editor world, or environment is already initialized. Doing nothing."));
 			return;
 		}
 
+		UE_LOG(LogSteamAudio, Log, TEXT("Initializing environment for [%s] world type [%i]."), *ListenerWorld->GetMapName(), (unsigned int)ListenerWorld->WorldType.GetValue());
+
+		CurrentLevel = ListenerWorld;
 		bool bIsUsingOcclusion = IsUsingSteamAudioPlugin(EAudioPlugin::OCCLUSION);
 		bool bIsUsingReverb = IsUsingSteamAudioPlugin(EAudioPlugin::REVERB);
 
+		// Get level name
+		FString CurrentLevelName;
+		CurrentLevelName = CurrentLevel->GetName();
+
 		if (bIsUsingOcclusion || bIsUsingReverb)
 		{
-			if (Environment.Initialize(ListenerWorld, AudioDevice))
+			if (Environment.Initialize(ListenerWorld, AudioDevice, &CurrentLevelName))
 			{
 				FScopeLock EnvironmentLock(Environment.GetEnvironmentCriticalSectionHandle());
 
-				UE_LOG(LogSteamAudio, Log, TEXT("Environment initialization successful."));
+				UE_LOG(LogSteamAudio, Log, TEXT("Environment initialization successful for level %s."), *CurrentLevelName);
 
 				if (bIsUsingReverb)
 				{
@@ -56,12 +69,83 @@ namespace SteamAudio
 					OcclusionPtr = static_cast<FPhononOcclusion*>(AudioDevice->OcclusionInterface.Get());
 					OcclusionPtr->SetEnvironment(&Environment);
 				}
+				
+				// Load dynamic scenes for the current level (if any)
+				TArray<FString> DynamicSceneFiles;
+				FString ScenePath = IFileManager::Get().ConvertToAbsolutePathForExternalAppForRead(*(DynamicRuntimePath / CurrentLevelName + "_*.phononscene"));
+
+				if (FPaths::DirectoryExists(DynamicRuntimePath))
+				{
+					IFileManager& FileManager = IFileManager::Get();
+					FileManager.FindFiles(DynamicSceneFiles, *ScenePath, true, false);
+
+					for (auto DynamicSceneFile : DynamicSceneFiles)
+					{
+						IPLhandle DynamicPhononScene = nullptr;
+						FPhononSceneInfo PhononSceneInfo;
+						FString DynamicSceneFileName = FPaths::GetBaseFilename(DynamicSceneFile);
+
+						// Get Dynamic Scene File
+						if (!LoadDynamicSceneFromDisk(CurrentLevel, Environment.GetComputeDevice(), Environment.GetSimulationSettings(), &DynamicPhononScene, PhononSceneInfo, DynamicSceneFileName))
+						{
+							UE_LOG(LogSteamAudio, Warning, TEXT("Unable to load dynamic scene file [%s]. Do a Scene Export and ensure your ray tracer is supported (Embree)."), *DynamicSceneFileName);
+							continue;
+						}
+
+						// Check if this scene have dynamic meshes to be added
+						bool bFoundInMap = false;
+						for (int32 i = 0; i < DynamicGeometry.Num(); i++)
+						{
+							if (DynamicGeometry[i].ID.Equals(DynamicSceneFileName)
+								&& DynamicGeometry[i].DynamicGeometryComponent
+								&& DynamicGeometry[i].DynamicGeometryComponent->IsValidLowLevelFast()
+								)
+							{
+								// Get 4x4 Matrix (1D array) from world transform
+								float Matrix[16];
+								auto DynGeoTransform = DynamicGeometry[i].DynamicGeometryComponent->GetAttachParent()->GetComponentTransform();
+								
+								GetMatrixForTransform(DynGeoTransform, Matrix, true, false);
+								IPLMatrix4x4 PhononMatrix = GetIPLMatrix(Matrix);
+								//DrawDebugSphere(ListenerWorld, DynamicGeometry[i].DynamicGeometryComponent->GetOwner()->GetRootComponent()->GetComponentTransform().GetTranslation(), 20.f, 4, FColor::Green, 9999.f);
+
+								// Create SteamAudio Instanced Mesh
+								IPLhandle InstancedMesh = nullptr;
+								IPLerror ErrCreateMesh = iplCreateInstancedMesh(Environment.GetScene(), DynamicPhononScene, PhononMatrix, &InstancedMesh);
+
+								if (ErrCreateMesh == IPL_STATUS_SUCCESS)
+								{
+									// Add Instanced Mesh to scene
+									iplAddInstancedMesh(Environment.GetScene(), InstancedMesh);
+
+									// Add references to map
+									DynamicGeometry[i].DynamicScene = DynamicPhononScene;
+									DynamicGeometry[i].DynamicGeometry = InstancedMesh;
+									UE_LOG(LogSteamAudio, Log, TEXT("Created instanced mesh [%s]"), *DynamicSceneFileName);
+								}
+								else
+								{
+									UE_LOG(LogSteamAudio, Warning, TEXT("Unable to create instanced mesh for [%s]"), *DynamicSceneFileName);
+								}
+
+								bFoundInMap = true;
+								break;
+							}
+						}
+
+						if (!bFoundInMap)
+						{
+							DynamicSceneMap.Add(DynamicSceneFileName, DynamicPhononScene);
+							UE_LOG(LogSteamAudio, Warning, TEXT("Loaded dynamic scene file [%s] but no dynamic geometry found in this level."), *DynamicSceneFileName);
+						}
+					}
+				}
 
 				bEnvironmentInitialized = true;
 			}
 			else
 			{
-				UE_LOG(LogSteamAudio, Warning, TEXT("Environment initialization unsuccessful."));
+				UE_LOG(LogSteamAudio, Warning, TEXT("Environment initialization unsuccessful for Level %s."), *CurrentLevelName);
 			}
 		}
 	}
@@ -102,10 +186,43 @@ namespace SteamAudio
 			Module->UnregisterAudioDevice(AudioDevice);
 		}
 	}
+	
+	void FPhononPluginManager::OnTick(UWorld* InWorld, const int32 ViewportIndex, const FTransform& ListenerTransform, const float InDeltaSeconds)
+	{
+		// Update any dynamic geometry
+		if (DynamicGeometry.Num() > 0)
+		{
+			// Update the transforms of the dynamic geometry
+			for (auto Geometry : DynamicGeometry)
+			{
+				// Check if we both have a dynamic mesh and dynamic scene
+				if (Geometry.DynamicGeometry
+					&& Geometry.DynamicScene
+					&& Geometry.DynamicGeometryComponent
+					&& Geometry.DynamicGeometryComponent->IsValidLowLevelFast()
+					)
+				{
+					// Get 4x4 Matrix (1D array) from world transform
+					float Matrix[16];
+					FTransform DynGeoTransform = Geometry.DynamicGeometryComponent->GetAttachParent()->GetComponentTransform();
+
+					GetMatrixForTransform(DynGeoTransform, Matrix, true, false);
+					IPLMatrix4x4 PhononMatrix = GetIPLMatrix(Matrix);
+					//DrawDebugSphere(InWorld, DynGeoTransform.GetTranslation(), 20.f, 4, FColor::Red);
+
+					// Update instanced mesh transform
+					iplUpdateInstancedMeshTransform(Geometry.DynamicGeometry, PhononMatrix);
+
+					// Commit changes to the main scene
+					iplCommitScene(Environment.GetScene());
+				}
+			}
+		}
+	}
 
 	void FPhononPluginManager::OnWorldChanged(FAudioDevice* AudioDevice, UWorld* ListenerWorld)
 	{
-		UE_LOG(LogSteamAudio, Log, TEXT("World changed. Reinitializing environment."));
+		UE_LOG(LogSteamAudio, Log, TEXT("World changed to %s. Reinitializing environment."), *ListenerWorld->GetMapName());
 
 		if (bEnvironmentInitialized)
 		{
@@ -130,4 +247,62 @@ namespace SteamAudio
 		FString CurrentPluginName = AudioPluginUtilities::GetDesiredPluginName(PluginType);
 		return CurrentPluginName.Equals(SteamPluginName);
 	}
+	
+	void FPhononPluginManager::AddDynamicGeometry(UPhononGeometryComponent* PhononGeometryComponent)
+	{
+		check(PhononGeometryComponent);
+
+		// Check if there's a valid world and the geometry component being registered is valid
+		if (PhononGeometryComponent->IsValidLowLevel())
+		{
+			FString DynamicGeometryID = PhononGeometryComponent->GetWorld()->GetName() + "_" + PhononGeometryComponent->GetOwner()->GetName() + "_" + PhononGeometryComponent->GetAttachParent()->GetName();
+			DynamicGeometry.Add({ DynamicGeometryID, PhononGeometryComponent, nullptr, nullptr });
+			UE_LOG(LogSteamAudio, Log, TEXT("Dynamic Geometry Added [%s] Current Count: %i"), *DynamicGeometryID, DynamicGeometry.Num());
+		}
+	}
+
+
+	void FPhononPluginManager::RemoveDynamicGeometry(UPhononGeometryComponent* PhononGeometryComponent)
+	{
+		check(PhononGeometryComponent);
+
+		// Check if there's a valid world and the geometry component being registered is valid
+		if (PhononGeometryComponent->IsValidLowLevel())
+		{
+			FString DynamicGeometryID = PhononGeometryComponent->GetWorld()->GetName() + "_" + PhononGeometryComponent->GetOwner()->GetName() + "_" + PhononGeometryComponent->GetAttachParent()->GetName();
+
+			// Check if this scene have dynamic meshes to be added
+			bool bFoundInMap = false;
+			for (int32 i = 0; i < DynamicGeometry.Num(); i++)
+			{
+				if (DynamicGeometry[i].ID.Equals(DynamicGeometryID)
+					&& DynamicGeometry[i].DynamicGeometryComponent
+					&& DynamicGeometry[i].DynamicGeometryComponent->IsValidLowLevelFast()
+					)
+				{
+					// Delete ipl instances
+					if (DynamicGeometry[i].DynamicGeometry)
+					{
+						iplDestroyStaticMesh(&(DynamicGeometry[i].DynamicGeometry));
+					}
+
+					if (DynamicGeometry[i].DynamicScene)
+					{
+						iplDestroyScene(&(DynamicGeometry[i].DynamicScene));
+					}
+
+					// Remove map entry
+					DynamicGeometry.RemoveAt(i);
+				}
+			}
+
+		}
+
+	}
+	
+	FEnvironment& FPhononPluginManager::GetPhononEnvironment()
+	{
+		return Environment;
+	}
+	
 }
diff --git a/Engine/Plugins/Runtime/Steam/SteamAudio/Source/SteamAudio/Private/PhononPluginManager.h b/Engine/Plugins/Runtime/Steam/SteamAudio/Source/SteamAudio/Private/PhononPluginManager.h
--- a/Engine/Plugins/Runtime/Steam/SteamAudio/Source/SteamAudio/Private/PhononPluginManager.h
+++ b/Engine/Plugins/Runtime/Steam/SteamAudio/Source/SteamAudio/Private/PhononPluginManager.h
@@ -9,7 +9,9 @@
 #include "PhononOcclusion.h"
 #include "PhononReverb.h"
 #include "SteamAudioEnvironment.h"
+#include "PhononCommon.h"
 #include "AudioPluginUtilities.h"
+#include "PhononGeometryComponent.h"
 
 namespace SteamAudio
 {
@@ -23,15 +25,37 @@ namespace SteamAudio
 	{
 	public:
 		FPhononPluginManager();
+		FPhononPluginManager(AActor* SubLevelRef);
 		~FPhononPluginManager();
 
 		//~ Begin IAudioPluginListener
 		virtual void OnListenerInitialize(FAudioDevice* AudioDevice, UWorld* ListenerWorld) override;
+		virtual void OnTick(UWorld* InWorld, const int32 ViewportIndex, const FTransform& ListenerTransform, const float InDeltaSeconds) override;
 		virtual void OnListenerUpdated(FAudioDevice* AudioDevice, const int32 ViewportIndex, const FTransform& ListenerTransform, const float InDeltaSeconds) override;
 		virtual void OnListenerShutdown(FAudioDevice* AudioDevice) override;
 		virtual void OnWorldChanged(FAudioDevice* AudioDevice, UWorld* ListenerWorld) override;
 		//~ End IAudioPluginListener
 
+		/** Cache of dynamic geometry currently active in the level. UPROPERTY to prevent GC */
+		UPROPERTY()
+		static TArray<FDynamicGeometryMap> DynamicGeometry;
+
+		/** Cache of dynamic geometry scenes currently active in the level. UPROPERTY to prevent GC */
+		UPROPERTY()
+		static TMap<FString, IPLhandle> DynamicSceneMap;
+
+		/** Add dynamic geometry to runtime cache */
+		static void AddDynamicGeometry(UPhononGeometryComponent* PhononGeometryComponent);
+
+		/** Remove dynamic geometry to runtime cache */
+		static void RemoveDynamicGeometry(UPhononGeometryComponent* PhononGeometryComponent);
+
+		/** Return the current phonon environment */
+		FEnvironment& GetPhononEnvironment();
+
+		/** Used to identify this plugin from other listeners - better if IAudioPluginListener.h has a GetPluginType() function instead */
+		static const FName PluginID;
+
 	private:
 		void InitializeEnvironment(FAudioDevice* AudioDevice, UWorld* ListenerWorld);
 
@@ -39,9 +63,12 @@ namespace SteamAudio
 		FEnvironment Environment;
 		FPhononReverb* ReverbPtr;
 		FPhononOcclusion* OcclusionPtr;
+		AActor* SubLevelReference = nullptr;
 
 	private:
 		/* Helper function for checking whether the user is using Steam Audio for spatialization, reverb, and/or occlusion: */
 		static bool IsUsingSteamAudioPlugin(EAudioPlugin PluginType);
+		static UWorld* CurrentLevel;
+
 	};
 }
diff --git a/Engine/Plugins/Runtime/Steam/SteamAudio/Source/SteamAudio/Private/PhononScene.cpp b/Engine/Plugins/Runtime/Steam/SteamAudio/Source/SteamAudio/Private/PhononScene.cpp
--- a/Engine/Plugins/Runtime/Steam/SteamAudio/Source/SteamAudio/Private/PhononScene.cpp
+++ b/Engine/Plugins/Runtime/Steam/SteamAudio/Source/SteamAudio/Private/PhononScene.cpp
@@ -9,6 +9,7 @@
 #include "PhononMaterialComponent.h"
 
 #include "EngineUtils.h"
+#include "HAL/FileManagerGeneric.h"
 #include "Engine/StaticMeshActor.h"
 #include "Engine/StaticMesh.h"
 #include "StaticMeshResources.h"
@@ -47,9 +48,19 @@
 namespace SteamAudio
 {
 	bool LoadSceneFromDisk(UWorld* World, IPLhandle ComputeDevice, const IPLSimulationSettings& SimulationSettings, IPLhandle* PhononScene,
-		FPhononSceneInfo& PhononSceneInfo)
+		FPhononSceneInfo& PhononSceneInfo, FString* CurrentLevelName)
 	{
-		FString MapName = StrippedMapName(World->GetMapName());
+		FString MapName;
+		
+		if (!CurrentLevelName || CurrentLevelName->IsEmpty())
+		{
+			MapName = World->GetMapName();
+		}
+		else
+		{
+			MapName = StrippedMapName(*CurrentLevelName);
+		}
+		
 		FString SceneFileName = RuntimePath + MapName + ".phononscene";
 		FString SceneInfoFileName = EditorOnlyPath + MapName + ".phononsceneinfo";
 		TArray<uint8> SceneData;
@@ -71,7 +82,7 @@ namespace SteamAudio
 		IPLerror Error = iplLoadScene(GlobalContext, SimulationSettings, SceneData.GetData(), SceneData.Num(), ComputeDevice, nullptr, PhononScene);
 		if (Error != IPL_STATUS_SUCCESS)
 		{
-			UE_LOG(LogSteamAudio, Warning, TEXT("Unable to load Phonon scene: error loading file."));
+			UE_LOG(LogSteamAudio, Warning, TEXT("Unable to load Phonon scene: error [%i] loading scene file [%s] with data size [%i]. \n "), (int32) Error, *SceneFileName, SceneData.Num());
 			return false;
 		}
 
@@ -90,8 +101,9 @@ namespace SteamAudio
 		if (SceneInfoFileHandle)
 		{
 			SceneInfoFileHandle->Read((uint8*)&PhononSceneInfo.NumTriangles, 4);
-			SceneInfoFileHandle->Seek(4);
+			SceneInfoFileHandle->Read((uint8*)&PhononSceneInfo.NumDynTriangles, 4);
 			SceneInfoFileHandle->Read((uint8*)&PhononSceneInfo.DataSize, 4);
+			SceneInfoFileHandle->Read((uint8*)&PhononSceneInfo.DynDataSize, 4);
 			delete SceneInfoFileHandle;
 		}
 		else
@@ -103,6 +115,68 @@ namespace SteamAudio
 		return true;
 	}
 
+	bool LoadDynamicSceneFromDisk(UWorld* World, IPLhandle ComputeDevice, const IPLSimulationSettings& SimulationSettings, IPLhandle* PhononScene,
+		FPhononSceneInfo& PhononSceneInfo, FString DynamicSceneID)
+	{
+		// Check for supported raytracer
+		if (GetDefault<USteamAudioSettings>()->RayTracer != EIplRayTracerType::EMBREE)
+		{
+			return false;
+		}
+
+		FString SceneFileName = DynamicRuntimePath + DynamicSceneID + ".phononscene";
+		FString SceneInfoFileName = DynamicEditorOnlyPath + DynamicSceneID + ".phononinfo";
+		TArray<uint8> SceneData;
+
+		IPlatformFile& PlatformFile = FPlatformFileManager::Get().GetPlatformFile();
+
+		IFileHandle* SceneFileHandle = PlatformFile.OpenRead(*SceneFileName);
+		if (SceneFileHandle)
+		{
+			SceneData.SetNum(SceneFileHandle->Size());
+			SceneFileHandle->Read(SceneData.GetData(), SceneFileHandle->Size());
+			delete SceneFileHandle;
+		}
+		else
+		{
+			UE_LOG(LogSteamAudio, Warning, TEXT("Unable to load Dynamic Phonon scene: error reading file. Be sure to export the scene."));
+			return false;
+		}
+
+		IPLerror Error = iplLoadScene(GlobalContext, SimulationSettings, SceneData.GetData(), SceneData.Num(), ComputeDevice, nullptr, PhononScene);
+		if (Error != IPL_STATUS_SUCCESS)
+		{
+			UE_LOG(LogSteamAudio, Warning, TEXT("Unable to load Dynamic Phonon scene: error [%i] loading scene file [%s] with data size [%i]. \n "), (int32)Error, *SceneFileName, SceneData.Num());
+			return false;
+		}
+
+		LoadDynamicScenesInfoFromDisk(World, PhononSceneInfo, DynamicSceneID);
+
+		return true;
+	}
+
+	bool LoadDynamicScenesInfoFromDisk(UWorld* World, FPhononSceneInfo& PhononSceneInfo, FString DynamicSceneID)
+	{
+		FString SceneInfoFileName = DynamicEditorOnlyPath + DynamicSceneID + ".phononinfo";
+
+		IPlatformFile& PlatformFile = FPlatformFileManager::Get().GetPlatformFile();
+		IFileHandle* SceneInfoFileHandle = PlatformFile.OpenRead(*SceneInfoFileName);
+		if (SceneInfoFileHandle)
+		{
+			SceneInfoFileHandle->Read((uint8*)&PhononSceneInfo.NumTriangles, 4);
+			SceneInfoFileHandle->Seek(4);
+			SceneInfoFileHandle->Read((uint8*)&PhononSceneInfo.DataSize, 4);
+			delete SceneInfoFileHandle;
+		}
+		else
+		{
+			UE_LOG(LogSteamAudio, Warning, TEXT("Unable to load Dynamic Phonon scene info: error loading info file."));
+			return false;
+		}
+
+		return true;
+	}
+
 #if WITH_EDITOR
 
 	static void LoadBSPGeometry(UWorld* World, IPLhandle PhononScene, TArray<IPLVector3>& IplVertices, TArray<IPLTriangle>& IplTriangles,
@@ -113,7 +187,13 @@ namespace SteamAudio
 		TArray<IPLint32>& IplMaterialIndices);
 	static void RegisterStaticMesh(IPLhandle PhononScene, TArray<IPLVector3>& IplVertices, TArray<IPLTriangle>& IplTriangles,
 		TArray<IPLint32>& IplMaterialIndices, IPLhandle* PhononStaticMesh);
+	static void RegisterStaticMeshes(IPLhandle PhononScene, TArray<IPLVector3>& IplVertices, TArray<IPLTriangle>& IplTriangles,
+		TArray<IPLint32>& IplMaterialIndices, TArray<IPLhandle>* PhononStaticMeshes);
 	static void GetPhononMaterials(UWorld* World, TArray<IPLMaterial>& IplMaterials);
+	static UPhononMaterialComponent* GetPhononMaterialComponent(AActor* Actor);
+	static uint32 GetMeshVerts(UStaticMeshComponent* StaticMeshComponent, TArray<IPLVector3>& VertexArray, bool bGetRelativeVertexPos = true);
+	static void CalculateVertsTrisMats(AActor* MeshActor, UStaticMeshComponent* MeshComponent, TArray<IPLVector3>& IplVertices, TArray<IPLTriangle>& IplTriangles,
+		TArray<IPLint32>& IplMaterialIndices);
 
 	//==============================================================================================================================================
 	// High level scene export
@@ -122,7 +202,7 @@ namespace SteamAudio
 	/**
 	 * Loads scene geometry, providing handles to the Phonon scene object and Phonon static meshes.
 	 */
-	bool CreateScene(UWorld* World, IPLhandle* PhononScene, IPLhandle* PhononStaticMesh, uint32& NumSceneTriangles)
+	bool CreateScene(UWorld* World, IPLhandle* PhononScene, IPLhandle* PhononStaticMesh, uint32& NumSceneTriangles, uint32& NumDynSceneTriangles, uint32& DynDataSize, bool CreateDynObjFiles /* =false*/ )
 	{
 		check(World);
 		check(PhononScene);
@@ -132,7 +212,7 @@ namespace SteamAudio
 
 		TPromise<bool> Result;
 
-		AsyncTask(ENamedThreads::GameThread, [World, PhononScene, PhononStaticMesh, &NumSceneTriangles, &Result]()
+		AsyncTask(ENamedThreads::GameThread, [World, PhononScene, PhononStaticMesh, &NumSceneTriangles, &NumDynSceneTriangles, &DynDataSize, CreateDynObjFiles, &Result]()
 		{
 			// Setup simulation settings - we really only care about scene type here, but Beta 17 API throws if ambisonics order > max order,
 			// so just zero the struct for now.
@@ -148,13 +228,22 @@ namespace SteamAudio
 			SceneMaterials.Add(GetDefault<USteamAudioSettings>()->GetDefaultStaticMeshMaterial());
 			SceneMaterials.Add(GetDefault<USteamAudioSettings>()->GetDefaultBSPMaterial());
 			SceneMaterials.Add(GetDefault<USteamAudioSettings>()->GetDefaultLandscapeMaterial());
-
+			
 			// Get IPLMaterials in scene
 			GetPhononMaterials(World, SceneMaterials);
 			UE_LOG(LogSteamAudio, Log, TEXT("Loading Phonon scene found [%i] materials."), SceneMaterials.Num());
-
+			
+			// Ensure that the SteamAudio paths exist, create them if they don't
+			CheckSceneDirectory(RuntimePath);
+			CheckSceneDirectory(EditorOnlyPath);
+
+			// Create dynamic scene
+			CreateDynamicScene(World, SimulationSettings, SceneMaterials, NumDynSceneTriangles, DynDataSize, CreateDynObjFiles);
+		
+			// Create static scene
 			IPLerror IplResult = iplCreateScene(GlobalContext, nullptr, SimulationSettings, SceneMaterials.Num(),
 				SceneMaterials.GetData(), nullptr, nullptr, nullptr, nullptr, nullptr, PhononScene);
+
 			if (IplResult != IPL_STATUS_SUCCESS)
 			{
 				UE_LOG(LogSteamAudio, Warning, TEXT("Error creating Phonon scene."));
@@ -177,7 +266,7 @@ namespace SteamAudio
 			{
 				LoadBSPGeometry(World, *PhononScene, IplVertices, IplTriangles, IplMaterialIndices);
 			}
-
+			
 			RegisterStaticMesh(*PhononScene, IplVertices, IplTriangles, IplMaterialIndices, PhononStaticMesh);
 			NumSceneTriangles = IplTriangles.Num();
 
@@ -190,6 +279,116 @@ namespace SteamAudio
 		return Value.Get();
 	}
 
+	void CreateDynamicScene(UWorld* World, IPLSimulationSettings SimulationSettings, TArray<IPLMaterial> SceneMaterials, uint32& NumSceneTriangles, uint32& DataSize, bool CreateObjFiles /* =false */)
+	{
+		// Wipe old dynamic scene files for this level (if any)
+		DeleteSceneFiles(World, DynamicRuntimePath);
+		DeleteSceneFiles(World, DynamicEditorOnlyPath);
+
+		for (TActorIterator<AActor> ActorItr(World); ActorItr; ++ActorItr)
+		{
+			if (ActorItr && ActorItr->IsValidLowLevelFast())
+			{
+				// Check if the actor's root component is movable and there's a phonon geometry component and at least one static mesh component attached
+				auto PhononGeometryComponent = ActorItr->GetComponentByClass(UPhononGeometryComponent::StaticClass());
+				if (ActorItr->GetRootComponent()
+					&& ActorItr->GetRootComponent()->Mobility == EComponentMobility::Movable
+					&& PhononGeometryComponent
+					&& PhononGeometryComponent->IsValidLowLevelFast()
+					&& ActorItr->GetComponentByClass(UStaticMeshComponent::StaticClass())
+					)
+				{
+					// Go through all static mesh components in this actor
+					TArray<UActorComponent*> Components;
+					ActorItr->GetComponents(UStaticMeshComponent::StaticClass(), Components);
+
+					for (auto ActorItrComponent : Components)
+					{
+						auto DynamicMeshComponent = Cast<UStaticMeshComponent>(ActorItrComponent);
+
+						if (DynamicMeshComponent && DynamicMeshComponent->IsValidLowLevelFast())
+						{
+							UPhononGeometryComponent* DynamicGeometryComponent = nullptr;
+							TArray<USceneComponent*> ChildrenComponents;
+							DynamicMeshComponent->GetChildrenComponents(false, ChildrenComponents);
+
+							// Check if this static mesh component has a Phonon Geometry Component
+							for (auto ChildComponent : ChildrenComponents)
+							{
+								auto ChildPhononGeometryComponent = Cast<UPhononGeometryComponent>(ChildComponent);
+
+								if (ChildPhononGeometryComponent)
+								{
+									DynamicGeometryComponent = ChildPhononGeometryComponent;
+									break;
+								}
+							}
+
+							if (DynamicGeometryComponent)
+							{
+								// Create an iplScene 
+								IPLhandle DynamicPhononScene = nullptr;
+								FString SceneID = World->GetName() + "_" + ActorItr->GetName() + "_" + DynamicMeshComponent->GetName();
+
+								IPLerror IplDynResult = iplCreateScene(GlobalContext, nullptr, SimulationSettings, SceneMaterials.Num(),
+									SceneMaterials.GetData(), nullptr, nullptr, nullptr, nullptr, nullptr, &DynamicPhononScene);
+								if (IplDynResult == IPL_STATUS_SUCCESS)
+								{
+									UE_LOG(LogSteamAudio, Log, TEXT("Dynamic Phonon scene succesfully created for [%s]"), *SceneID);
+								}
+								else
+								{
+									UE_LOG(LogSteamAudio, Warning, TEXT("Error creating Dynamic Phonon scene for [%s]"), *SceneID);
+									continue;
+								}
+
+								// Create and register dynamic meshes to the dynamic scene
+								uint32 DynamicSceneTriangles = 0;
+								TArray<IPLhandle> PhononDynamicMeshes;
+								TArray<IPLVector3> DynIplVertices;
+								TArray<IPLTriangle> DynIplTriangles;
+								TArray<IPLint32> DynIplMaterialIndices;
+
+
+								CalculateVertsTrisMats(*ActorItr, DynamicMeshComponent, DynIplVertices, DynIplTriangles, DynIplMaterialIndices);
+
+								// Register Instanced Mesh to Dynamic Scene
+								RegisterStaticMeshes(DynamicPhononScene, DynIplVertices, DynIplTriangles, DynIplMaterialIndices, &PhononDynamicMeshes);
+
+								// Add tris in this mesh to the total scene tris
+								DynamicSceneTriangles += DynIplTriangles.Num();
+
+								// Finalize and save to disk
+								FPhononSceneInfo DynamicPhononSceneInfo;
+								DynamicPhononSceneInfo.NumTriangles = DynamicSceneTriangles;
+								DynamicPhononSceneInfo.DataSize = iplSaveScene(DynamicPhononScene, nullptr);
+								bool SaveDynSceneSuccessful = SaveDynamicSceneToDisk(World, DynamicPhononScene, DynamicPhononSceneInfo, SceneID, CreateObjFiles);
+
+								// Clean up Dynamic Phonon structures
+								for (IPLhandle PhononDynamicMeshItr : PhononDynamicMeshes)
+								{
+									iplDestroyStaticMesh(&PhononDynamicMeshItr);
+								}
+								iplDestroyScene(&DynamicPhononScene);
+
+								if (SaveDynSceneSuccessful)
+								{
+									DataSize += DynamicPhononSceneInfo.DataSize;
+									NumSceneTriangles += DynamicPhononSceneInfo.NumTriangles;
+									UE_LOG(LogSteamAudio, Log, TEXT("Dynamic scene [%s] with [%i] triangles saved to disk with size [%i]."), *SceneID, DynamicSceneTriangles, DynamicPhononSceneInfo.DataSize);
+								}
+								else
+								{
+									UE_LOG(LogSteamAudio, Warning, TEXT("Unable to save Dynamic scene [%s] with [%i] triangles to disk with size [%i]."), *SceneID, DynamicSceneTriangles, DynamicPhononSceneInfo.DataSize);
+								}
+							}
+						}
+					}
+				}
+			}
+		}
+	}
+
 	bool SaveFinalizedSceneToDisk(UWorld* World, IPLhandle PhononScene, const FPhononSceneInfo& PhononSceneInfo)
 	{
 		// Write Phonon Scene data to byte array
@@ -210,7 +409,9 @@ namespace SteamAudio
 			delete SceneFileHandle;
 
 			SceneInfoHandle->Write((uint8*)&PhononSceneInfo.NumTriangles, 4);
+			SceneInfoHandle->Write((uint8*)&PhononSceneInfo.NumDynTriangles, 4);
 			SceneInfoHandle->Write((uint8*)&PhononSceneInfo.DataSize, 4);
+			SceneInfoHandle->Write((uint8*)&PhononSceneInfo.DynDataSize, 4);
 			delete SceneInfoHandle;
 		}
 		else
@@ -231,6 +432,88 @@ namespace SteamAudio
 		return true;
 	}
 
+
+	bool STEAMAUDIO_API SaveDynamicSceneToDisk(UWorld* World, IPLhandle PhononScene, const FPhononSceneInfo& PhononSceneInfo, FString DynamicSceneID, bool CreateObjFile /* =false */)
+	{
+		// Write Phonon Scene data to byte array
+		TArray<uint8> SceneData;
+		SceneData.SetNum(PhononSceneInfo.DataSize);
+		iplSaveScene(PhononScene, SceneData.GetData());
+
+		// Serialize byte array to disk
+		FString SceneFileName = DynamicRuntimePath + DynamicSceneID + ".phononscene";
+		FString SceneInfoName = DynamicEditorOnlyPath + DynamicSceneID + ".phononinfo";
+		IPlatformFile& PlatformFile = FPlatformFileManager::Get().GetPlatformFile();
+		IFileHandle* SceneFileHandle = PlatformFile.OpenWrite(*SceneFileName);
+		IFileHandle* SceneInfoHandle = PlatformFile.OpenWrite(*SceneInfoName);
+
+		if (SceneFileHandle && SceneInfoHandle)
+		{
+			if (CreateObjFile)
+			{
+				iplSaveSceneAsObj(PhononScene, TCHAR_TO_ANSI(*(DynamicEditorOnlyPath + DynamicSceneID + ".obj")));
+			}
+
+			SceneFileHandle->Write(SceneData.GetData(), PhononSceneInfo.DataSize);
+			delete SceneFileHandle;
+
+			SceneInfoHandle->Write((uint8*)&PhononSceneInfo.NumTriangles, 4);
+			SceneInfoHandle->Write((uint8*)&PhononSceneInfo.DataSize, 4);
+			delete SceneInfoHandle;
+		}
+		else
+		{
+			if (SceneFileHandle)
+			{
+				delete SceneFileHandle;
+			}
+
+			if (SceneInfoHandle)
+			{
+				delete SceneInfoHandle;
+			}
+
+			return false;
+		}
+
+		return true;
+	}
+
+	bool STEAMAUDIO_API CheckSceneDirectory(FString DirectoryPath)
+	{
+		IFileManager& FileManager = FFileManagerGeneric::Get();
+		if (!FileManager.DirectoryExists(*DirectoryPath))
+		{
+			FileManager.MakeDirectory(*DirectoryPath);
+			return false;
+		}
+		
+		return true;
+	}
+
+	void STEAMAUDIO_API DeleteSceneFiles(UWorld * World, FString FilePath)
+	{
+		IFileManager& FileManager = FFileManagerGeneric::Get();
+		if (CheckSceneDirectory(FilePath))
+		{
+			TArray<FString> FoundFiles;
+			FileManager.FindFiles(FoundFiles, *FilePath);
+
+			for (auto FoundFile : FoundFiles)
+			{
+				FString FileToDelete = FilePath + FoundFile;
+				//UE_LOG(LogSteamAudio, Warning, TEXT("Checking for file [%s]"), *FileToDelete);
+
+				if (FoundFile.Left(World->GetName().Len()).Equals(World->GetName(), ESearchCase::IgnoreCase) 
+					&& FileManager.Delete(*FileToDelete, true, true, true)
+					)
+				{
+					UE_LOG(LogSteamAudio, Log, TEXT("Deleted old scene file [%s]"), *FileToDelete);
+				}
+			}
+		}
+	}
+
 	//==============================================================================================================================================
 	// Utilities for adding/removing Phonon Geometry components
 	//==============================================================================================================================================
@@ -244,12 +527,7 @@ namespace SteamAudio
 
 		for (TActorIterator<AActor> ActorItr(World); ActorItr; ++ActorItr)
 		{
-			// Check if there's already a phonon geometry component in this actor
-			if (ActorItr->GetComponentByClass(UPhononGeometryComponent::StaticClass()))
-			{
-				continue;
-			}
-
+			
 			// Go through any static mesh components found in this actor
 			TArray<UActorComponent*> Components;
 			ActorItr->GetComponents(UStaticMeshComponent::StaticClass(), Components);
@@ -262,17 +540,35 @@ namespace SteamAudio
 	
 					// Check if this static mesh can affect audio
 					if (StaticMeshComp
-						&& StaticMeshComp->IsValidLowLevelFast()
-						&& StaticMeshComp->Mobility != EComponentMobility::Movable
-						&& StaticMeshComp->GetStaticMesh()
-						&& StaticMeshComp->GetStaticMesh()->HasValidRenderData()
+						&& StaticMeshComp->IsValidLowLevel()
 						&& !StaticMeshComp->IsVisualizationComponent()
 						)
 					{
-						// Add phonon geometry object
-						auto PhononGeometryComponent = NewObject<UPhononGeometryComponent>(*ActorItr);
-						PhononGeometryComponent->RegisterComponent();
-						ActorItr->AddInstanceComponent(PhononGeometryComponent);
+						// Check if there's already a phonon geometry component attached to this static mesh component
+						TArray<USceneComponent*> ChildrenComponents;
+						StaticMeshComp->GetChildrenComponents(true, ChildrenComponents);
+						bool bAlreadyHavePhononGeomComponent = false;
+						
+						for (auto ChildComponent : ChildrenComponents)
+						{
+							auto* PhononGeomComp = Cast<UPhononGeometryComponent>(ChildComponent);
+							
+							if (PhononGeomComp && PhononGeomComp->IsValidLowLevelFast())
+							{
+								bAlreadyHavePhononGeomComponent = true;
+								break;
+							}
+						}
+
+						// If there isn't, create and attach a new phonon geometry component to this static mesh
+						if (!bAlreadyHavePhononGeomComponent)
+						{
+							// Add phonon geometry object
+							auto PhononGeometryComponent = NewObject<UPhononGeometryComponent>(*ActorItr);
+							PhononGeometryComponent->AttachToComponent(StaticMeshComp, FAttachmentTransformRules::KeepRelativeTransform);
+							PhononGeometryComponent->RegisterComponent();
+							ActorItr->AddInstanceComponent(PhononGeometryComponent);
+						}
 					}
 				}
 			}
@@ -303,13 +599,13 @@ namespace SteamAudio
 	/**
 	 * Populates VertexArray with the given mesh's vertices. Converts from UE4 coords to Phonon coords. Returns the number of vertices added.
 	 */
-	static uint32 GetMeshVerts(TActorIterator<AStaticMeshActor>& AStaticMeshItr, TArray<IPLVector3>& VertexArray)
+	static uint32 GetMeshVerts(UStaticMeshComponent* StaticMeshComponent, TArray<IPLVector3>& VertexArray, bool bGetRelativeVertexPos /*= true*/)
 	{
-		check(AStaticMeshItr->GetStaticMeshComponent()->GetStaticMesh()->HasValidRenderData());
+		check(StaticMeshComponent->GetStaticMesh()->HasValidRenderData());
 
 		uint32 NumVerts = 0;
 
-		FStaticMeshLODResources& LODModel = AStaticMeshItr->GetStaticMeshComponent()->GetStaticMesh()->RenderData->LODResources[0];
+		FStaticMeshLODResources& LODModel = StaticMeshComponent->GetStaticMesh()->RenderData->LODResources[0];
 		auto Indices = LODModel.IndexBuffer.GetArrayView();
 
 		for (const FStaticMeshSection& Section : LODModel.Sections)
@@ -320,7 +616,9 @@ namespace SteamAudio
 				for (auto v = 2; v > -1; v--)
 				{
 					auto i = Indices[BaseIndex + v];
-					auto vtx = AStaticMeshItr->ActorToWorld().TransformPosition(LODModel.VertexBuffers.PositionVertexBuffer.VertexPosition(i));
+					auto vtx = bGetRelativeVertexPos ? StaticMeshComponent->GetComponentTransform().TransformPosition(LODModel.VertexBuffers.PositionVertexBuffer.VertexPosition(i))
+						: LODModel.VertexBuffers.PositionVertexBuffer.VertexPosition(i);
+
 					VertexArray.Add(UnrealToPhononIPLVector3(vtx));
 					NumVerts++;
 				}
@@ -348,6 +646,9 @@ namespace SteamAudio
 		return false;
 	}
 
+	/**
+	 * Walks up the actor attachment chain, checking for a Phonon Material component.
+	 */
 	static UPhononMaterialComponent* GetPhononMaterialComponent(AActor* Actor)
 	{
 		auto CurrentActor = Actor;
@@ -363,6 +664,48 @@ namespace SteamAudio
 		return nullptr;
 	}
 
+	/**
+	 * Calculates a mesh's vertices, tris and retries phonon material indices for a given actor
+	 */
+	static void CalculateVertsTrisMats(AActor* MeshActor, UStaticMeshComponent* MeshComponent, TArray<IPLVector3>& IplVertices, TArray<IPLTriangle>& IplTriangles,
+		TArray<IPLint32>& IplMaterialIndices)
+	{
+		check(MeshActor);
+		check(MeshComponent);
+
+		int32 StartVertexIdx = IplVertices.Num();
+		int32 NumMeshVerts = MeshComponent->Mobility == EComponentMobility::Movable ? GetMeshVerts(MeshComponent, IplVertices, false) : GetMeshVerts(MeshComponent, IplVertices);
+		int32 NumMeshTriangles = NumMeshVerts / 3;
+
+		for (int32 i = 0; i < NumMeshTriangles; ++i)
+		{
+			IPLTriangle IplTriangle;
+			IplTriangle.indices[0] = StartVertexIdx + i * 3;
+			IplTriangle.indices[1] = StartVertexIdx + i * 3 + 2;
+			IplTriangle.indices[2] = StartVertexIdx + i * 3 + 1;
+			IplTriangles.Add(IplTriangle);
+		}
+
+		auto PhononMaterialComponent = GetPhononMaterialComponent(MeshActor);
+		auto MaterialIndex = 0;
+
+		if (PhononMaterialComponent)
+		{
+			MaterialIndex = PhononMaterialComponent->MaterialIndex;
+		}
+		else
+		{
+			// The default static mesh material is always registered at size(MaterialPresets)
+			MaterialIndex = MaterialPresets.Num();
+		}
+
+		for (auto i = 0; i < NumMeshTriangles; ++i)
+		{
+			IplMaterialIndices.Add(MaterialIndex);
+			//UE_LOG(LogSteamAudio, Warning, TEXT("Phonon scene Material Index Added [%i]"), MaterialIndex);
+		}
+	}
+
 	/**
 	 * Loads any static mesh actors.
 	 */
@@ -377,45 +720,53 @@ namespace SteamAudio
 		for (TActorIterator<AStaticMeshActor> AStaticMeshItr(World); AStaticMeshItr; ++AStaticMeshItr)
 		{
 			// Only consider static mesh actors that have both an acoustic geometry component attached and valid render data
-			if (IsActorPhononGeometry(*AStaticMeshItr) && AStaticMeshItr->GetStaticMeshComponent()->GetStaticMesh() &&
-				AStaticMeshItr->GetStaticMeshComponent()->GetStaticMesh()->HasValidRenderData())
+			if (IsActorPhononGeometry(*AStaticMeshItr) 
+				&& AStaticMeshItr->GetStaticMeshComponent()->Mobility != EComponentMobility::Movable
+				&& AStaticMeshItr->GetStaticMeshComponent()->GetStaticMesh() 
+				&& AStaticMeshItr->GetStaticMeshComponent()->GetStaticMesh()->HasValidRenderData()
+				)
 			{
-				auto PhononGeometryComponent = static_cast<UPhononGeometryComponent*>(
-					AStaticMeshItr->GetComponentByClass(UPhononGeometryComponent::StaticClass()));
-
-				int32 StartVertexIdx = IplVertices.Num();
-				int32 NumMeshVerts = GetMeshVerts(AStaticMeshItr, IplVertices);
-				int32 NumMeshTriangles = NumMeshVerts / 3;
-
-				for (int32 i = 0; i < NumMeshTriangles; ++i)
-				{
-					IPLTriangle IplTriangle;
-					IplTriangle.indices[0] = StartVertexIdx + i * 3;
-					IplTriangle.indices[1] = StartVertexIdx + i * 3 + 2;
-					IplTriangle.indices[2] = StartVertexIdx + i * 3 + 1;
-					IplTriangles.Add(IplTriangle);
-				}
+				CalculateVertsTrisMats(*AStaticMeshItr, AStaticMeshItr->GetStaticMeshComponent(), IplVertices, IplTriangles, IplMaterialIndices);
+			}
+		}
+	}
 
-				auto PhononMaterialComponent = GetPhononMaterialComponent(*AStaticMeshItr);
-				auto MaterialIndex = 0;
+	/**
+	 * Loads any static mesh actors, adding any Phonon static meshes to the provided array.
+	 */
+	static uint32 LoadDynamicMeshes(AActor* DynamicActor, IPLhandle PhononScene)
+	{
+		check(DynamicActor);
+		check(PhononScene);
 
-				if (PhononMaterialComponent)
-				{
-					MaterialIndex = PhononMaterialComponent->MaterialIndex;
-				}
-				else
-				{
-					// The default static mesh material is always registered at size(MaterialPresets)
-					MaterialIndex = MaterialPresets.Num();
-				}
+		TArray<IPLVector3> IplVertices;
+		TArray<IPLTriangle> IplTriangles;
+		TArray<IPLint32> IplMaterialIndices;
 
-				for (auto i = 0; i < NumMeshTriangles; ++i)
-				{
-					IplMaterialIndices.Add(MaterialIndex);
-					//UE_LOG(LogSteamAudio, Warning, TEXT("Phonon scene Material Index Added [%i]"), MaterialIndex);
-				}
+		// Go through all static mesh components in this actor
+		TArray<UActorComponent*> Components;
+		DynamicActor->GetComponents(UStaticMeshComponent::StaticClass(), Components);
+		for (auto StaticMeshComp : Components)
+		{
+			auto StaticMeshComponent = Cast<UStaticMeshComponent>(StaticMeshComp);
+
+			// Only consider dynamic static meshes that have a static mesh already assigned and have valid render data
+			if (StaticMeshComponent
+				&& StaticMeshComponent->IsValidLowLevelFast()
+				&& StaticMeshComponent->Mobility == EComponentMobility::Movable
+				&& StaticMeshComponent->GetStaticMesh()
+				&& StaticMeshComponent->GetStaticMesh()->HasValidRenderData()
+				&& !StaticMeshComponent->IsVisualizationComponent()
+				)
+			{
+				CalculateVertsTrisMats(DynamicActor, StaticMeshComponent, IplVertices, IplTriangles, IplMaterialIndices);
 			}
 		}
+
+		// Register a new static mesh with Phonon
+		//RegisterStaticMesh(PhononScene, IplVertices, IplTriangles, IplMaterialIndices, PhononStaticMeshes);
+
+		return IplTriangles.Num();
 	}
 
 	//==============================================================================================================================================
@@ -567,6 +918,33 @@ namespace SteamAudio
 		}
 	}
 
+	/**
+	 * Registers a new static mesh with Phonon, adding its handle to the provided array of static meshes.
+	 */
+	static void RegisterStaticMeshes(IPLhandle PhononScene, TArray<IPLVector3>& IplVertices, TArray<IPLTriangle>& IplTriangles,
+		TArray<IPLint32>& IplMaterialIndices, TArray<IPLhandle>* PhononStaticMeshes)
+	{
+		if (IplVertices.Num() > 0)
+		{
+			UE_LOG(LogSteamAudio, Log, TEXT("Registering new mesh with %d triangles."), IplTriangles.Num());
+
+			IPLhandle IplStaticMesh = nullptr;
+			auto IplResult = iplCreateStaticMesh(PhononScene, IplVertices.Num(), IplTriangles.Num(), IplVertices.GetData(), IplTriangles.GetData(),
+				IplMaterialIndices.GetData(), &IplStaticMesh);
+			if (IplResult != IPL_STATUS_SUCCESS)
+			{
+				UE_LOG(LogSteamAudio, Warning, TEXT("Error adding a new object to the acoustic scene."));
+				return;
+			}
+
+			PhononStaticMeshes->Add(IplStaticMesh);
+		}
+		else
+		{
+			UE_LOG(LogSteamAudio, Warning, TEXT("Skipping mesh registration because no vertices were found."));
+		}
+	}
+
 	/**
 	 * Calculates the total number of materials that must be registered with Phonon. This includes presets and any custom materials.
 	 */
diff --git a/Engine/Plugins/Runtime/Steam/SteamAudio/Source/SteamAudio/Private/PhononScene.h b/Engine/Plugins/Runtime/Steam/SteamAudio/Source/SteamAudio/Private/PhononScene.h
--- a/Engine/Plugins/Runtime/Steam/SteamAudio/Source/SteamAudio/Private/PhononScene.h
+++ b/Engine/Plugins/Runtime/Steam/SteamAudio/Source/SteamAudio/Private/PhononScene.h
@@ -9,6 +9,7 @@
 #include "Async/Async.h"
 #include "CoreMinimal.h"
 #include "UObject/ObjectMacros.h"
+#include "Engine/StaticMeshActor.h"
 
 class AActor;
 class AStaticMeshActor;
@@ -20,27 +21,43 @@ namespace SteamAudio
 	{
 		FPhononSceneInfo()
 			: NumTriangles(0)
+			, NumDynTriangles(0)
 			, DataSize(0)
+			, DynDataSize(0)
 		{}
 
 		uint32 NumTriangles;
+		uint32 NumDynTriangles;
 		uint32 DataSize;
+		uint32 DynDataSize;
 	};
 
 	bool STEAMAUDIO_API LoadSceneFromDisk(UWorld* World, IPLhandle ComputeDevice, const IPLSimulationSettings& SimulationSettings,
-		IPLhandle* PhononScene, FPhononSceneInfo& PhononSceneInfo);
+		IPLhandle* PhononScene, FPhononSceneInfo& PhononSceneInfo, FString* CurrentLevelName);
 
 	bool STEAMAUDIO_API LoadSceneInfoFromDisk(UWorld* World, FPhononSceneInfo& PhononSceneInfo);
 
+	bool STEAMAUDIO_API LoadDynamicSceneFromDisk(UWorld* World, IPLhandle ComputeDevice, const IPLSimulationSettings& SimulationSettings,
+		IPLhandle* PhononScene, FPhononSceneInfo& PhononSceneInfo, FString DynamicSceneID);
+
+	bool STEAMAUDIO_API LoadDynamicScenesInfoFromDisk(UWorld* World, FPhononSceneInfo& PhononSceneInfo, FString DynamicSceneID);
+
 #if WITH_EDITOR
 
 	/** Creates and populates a Phonon scene. Gathers Phonon Geometry asynchronously on the game thread. */
-	bool STEAMAUDIO_API CreateScene(UWorld* World, IPLhandle* PhononScene, IPLhandle* PhononStaticMesh, uint32& NumSceneTriangles);
+	bool STEAMAUDIO_API CreateScene(UWorld* World, IPLhandle* PhononScene, IPLhandle* PhononStaticMesh, uint32& NumSceneTriangles, uint32& NumDynSceneTriangles, uint32& DynDataSize, bool CreateDynObjFiles = false);
+	void STEAMAUDIO_API CreateDynamicScene(UWorld* World, IPLSimulationSettings SimulationSettings, TArray<IPLMaterial> SceneMaterials, uint32& NumSceneTriangles, uint32& DataSize, bool CreateObjFile = false);
 	bool STEAMAUDIO_API SaveFinalizedSceneToDisk(UWorld* World, IPLhandle PhononScene, const FPhononSceneInfo& PhononSceneInfo);
-
+	bool STEAMAUDIO_API SaveDynamicSceneToDisk(UWorld* World, IPLhandle PhononScene, const FPhononSceneInfo& PhononSceneInfo, FString DynamicSceneID, bool CreateObjFile = false);
 	void STEAMAUDIO_API AddGeometryComponentsToStaticMeshes(UWorld* World);
 	void STEAMAUDIO_API RemoveGeometryComponentsFromStaticMeshes(UWorld* World);
 
+	/** Check a directory if it exists, create it if it doesn't. Return True if the directory exists. */
+	bool STEAMAUDIO_API CheckSceneDirectory(FString DirectoryPath);
+
+	/** Delete any scene files in a given path that start with the world/level name. */
+	void STEAMAUDIO_API DeleteSceneFiles(UWorld* World, FString FilePath);
+
 #endif // WITH_EDITOR
 
 	uint32 GetNumTrianglesForStaticMesh(AStaticMeshActor* StaticMeshActor);
diff --git a/Engine/Plugins/Runtime/Steam/SteamAudio/Source/SteamAudio/Private/SteamAudioEnvironment.cpp b/Engine/Plugins/Runtime/Steam/SteamAudio/Source/SteamAudio/Private/SteamAudioEnvironment.cpp
--- a/Engine/Plugins/Runtime/Steam/SteamAudio/Source/SteamAudio/Private/SteamAudioEnvironment.cpp
+++ b/Engine/Plugins/Runtime/Steam/SteamAudio/Source/SteamAudio/Private/SteamAudioEnvironment.cpp
@@ -42,7 +42,7 @@ namespace SteamAudio
 	{
 	}
 
-	bool FEnvironment::Initialize(UWorld* World, FAudioDevice* InAudioDevice)
+	bool FEnvironment::Initialize(UWorld* World, FAudioDevice* InAudioDevice, FString* CurrentLevelName)
 	{
 		if (World == nullptr)
 		{
@@ -116,7 +116,7 @@ namespace SteamAudio
 		EnvironmentalOutputAudioFormat.ambisonicsNormalization = IPL_AMBISONICSNORMALIZATION_N3D;
 		EnvironmentalOutputAudioFormat.ambisonicsOrdering = IPL_AMBISONICSORDERING_ACN;
 
-		if (!LoadSceneFromDisk(World, ComputeDevice, SimulationSettings, &PhononScene, PhononSceneInfo))
+		if (!LoadSceneFromDisk(World, ComputeDevice, SimulationSettings, &PhononScene, PhononSceneInfo, CurrentLevelName))
 		{
 			UE_LOG(LogSteamAudio, Warning, TEXT("Unable to create Phonon environment: failed to load scene from disk. Be sure to export the scene."));
 			return false;
diff --git a/Engine/Plugins/Runtime/Steam/SteamAudio/Source/SteamAudio/Private/SteamAudioEnvironment.h b/Engine/Plugins/Runtime/Steam/SteamAudio/Source/SteamAudio/Private/SteamAudioEnvironment.h
--- a/Engine/Plugins/Runtime/Steam/SteamAudio/Source/SteamAudio/Private/SteamAudioEnvironment.h
+++ b/Engine/Plugins/Runtime/Steam/SteamAudio/Source/SteamAudio/Private/SteamAudioEnvironment.h
@@ -44,11 +44,12 @@ namespace SteamAudio
 		FEnvironment();
 		~FEnvironment();
 
-		bool Initialize(UWorld* World, FAudioDevice* InAudioDevice);
+		bool Initialize(UWorld* World, FAudioDevice* InAudioDevice, FString* CurrentLevelName);
 		void Shutdown();
 
-		IPLhandle GetScene() const { return PhononScene; };
+		IPLhandle GetScene() { return PhononScene; };
 		IPLhandle GetEnvironment() const { return PhononEnvironment; };
+		IPLhandle GetComputeDevice() const { return ComputeDevice; };
 		IPLhandle GetEnvironmentalRenderer() const { return EnvironmentalRenderer; }
 		const IPLSimulationSettings& GetSimulationSettings() const { return SimulationSettings; }
 		const IPLRenderingSettings& GetRenderingSettings() const { return RenderingSettings; }
diff --git a/Engine/Plugins/Runtime/Steam/SteamAudio/Source/SteamAudio/Private/SteamAudioModule.cpp b/Engine/Plugins/Runtime/Steam/SteamAudio/Source/SteamAudio/Private/SteamAudioModule.cpp
--- a/Engine/Plugins/Runtime/Steam/SteamAudio/Source/SteamAudio/Private/SteamAudioModule.cpp
+++ b/Engine/Plugins/Runtime/Steam/SteamAudio/Source/SteamAudio/Private/SteamAudioModule.cpp
@@ -12,12 +12,6 @@ IMPLEMENT_MODULE(SteamAudio::FSteamAudioModule, SteamAudio)
 namespace SteamAudio
 {
 	void* FSteamAudioModule::PhononDllHandle = nullptr;
-	void* FSteamAudioModule::TANDllHandle = nullptr;
-	void* FSteamAudioModule::TANUtilsDllHandle = nullptr;
-	void* FSteamAudioModule::EmbreeDllHandle = nullptr;
-	void* FSteamAudioModule::RadeonRaysDllHandle = nullptr;
-	void* FSteamAudioModule::TbbDllHandle = nullptr;
-	void* FSteamAudioModule::TbbMallocDllHandle = nullptr;
 
 	static bool bModuleStartedUp = false;
 
@@ -40,7 +34,9 @@ namespace SteamAudio
 		// Make folder paths globally available
 		BasePath = FPaths::ProjectContentDir() + "SteamAudio/";
 		RuntimePath = BasePath + "Runtime/";
+		DynamicRuntimePath = RuntimePath + "Dynamic/";
 		EditorOnlyPath = BasePath + "EditorOnly/";
+		DynamicEditorOnlyPath = EditorOnlyPath + "Dynamic/";
 
 		//Register the Steam Audio plugin factories
 		IModularFeatures::Get().RegisterModularFeature(FSpatializationPluginFactory::GetModularFeatureName(), &SpatializationPluginFactory);
@@ -54,12 +50,6 @@ namespace SteamAudio
 			FString PathToDll = FPaths::EngineDir() / TEXT("Binaries/ThirdParty/Phonon/Win32/");
 #else
 			FString PathToDll = FPaths::EngineDir() / TEXT("Binaries/ThirdParty/Phonon/Win64/");
-			FSteamAudioModule::TANDllHandle = LoadDll(PathToDll + TEXT("tanrt64.dll"), false);
-			FSteamAudioModule::TANUtilsDllHandle = LoadDll(PathToDll + TEXT("GPUUtilities.dll"), false);
-			FSteamAudioModule::RadeonRaysDllHandle = LoadDll(PathToDll + TEXT("RadeonRays.dll"), false);
-			FSteamAudioModule::EmbreeDllHandle = LoadDll(PathToDll + TEXT("embree.dll"), false);
-			FSteamAudioModule::TbbDllHandle = LoadDll(PathToDll + TEXT("tbb.dll"), false);
-			FSteamAudioModule::TbbMallocDllHandle = LoadDll(PathToDll + TEXT("tbbmalloc.dll"), false);
 #endif
 
 			FString DLLToLoad = PathToDll + TEXT("phonon.dll");
@@ -86,42 +76,6 @@ namespace SteamAudio
 			FPlatformProcess::FreeDllHandle(FSteamAudioModule::PhononDllHandle);
 			FSteamAudioModule::PhononDllHandle = nullptr;
 		}
-
-		if (FSteamAudioModule::TANDllHandle)
-		{
-			FPlatformProcess::FreeDllHandle(FSteamAudioModule::TANDllHandle);
-			FSteamAudioModule::TANDllHandle = nullptr;
-		}
-
-		if (FSteamAudioModule::TANUtilsDllHandle)
-		{
-			FPlatformProcess::FreeDllHandle(FSteamAudioModule::TANUtilsDllHandle);
-			FSteamAudioModule::TANUtilsDllHandle = nullptr;
-		}
-
-		if (FSteamAudioModule::EmbreeDllHandle)
-		{
-			FPlatformProcess::FreeDllHandle(FSteamAudioModule::EmbreeDllHandle);
-			FSteamAudioModule::EmbreeDllHandle = nullptr;
-		}
-
-		if (FSteamAudioModule::RadeonRaysDllHandle)
-		{
-			FPlatformProcess::FreeDllHandle(FSteamAudioModule::RadeonRaysDllHandle);
-			FSteamAudioModule::RadeonRaysDllHandle = nullptr;
-		}
-
-		if (FSteamAudioModule::TbbDllHandle)
-		{
-			FPlatformProcess::FreeDllHandle(FSteamAudioModule::TbbDllHandle);
-			FSteamAudioModule::TbbDllHandle = nullptr;
-		}
-
-		if (FSteamAudioModule::TbbMallocDllHandle)
-		{
-			FPlatformProcess::FreeDllHandle(FSteamAudioModule::TbbMallocDllHandle);
-			FSteamAudioModule::TbbMallocDllHandle = nullptr;
-		}
 #endif
 	}
 
@@ -167,7 +121,7 @@ namespace SteamAudio
 		{
 			Module->RegisterAudioDevice(OwningDevice);
 		}
-
+		
 		return TAudioOcclusionPtr(new FPhononOcclusion());
 	}
 
diff --git a/Engine/Plugins/Runtime/Steam/SteamAudio/Source/SteamAudio/SteamAudio.Build.cs b/Engine/Plugins/Runtime/Steam/SteamAudio/Source/SteamAudio/SteamAudio.Build.cs
--- a/Engine/Plugins/Runtime/Steam/SteamAudio/Source/SteamAudio/SteamAudio.Build.cs
+++ b/Engine/Plugins/Runtime/Steam/SteamAudio/Source/SteamAudio/SteamAudio.Build.cs
@@ -50,7 +50,11 @@ namespace UnrealBuildTool.Rules
 			{
 				PrivateDependencyModuleNames.Add("UnrealEd");
 				PrivateDependencyModuleNames.Add("Landscape");
-			}
+            }
+            else
+            {
+                RuntimeDependencies.Add("$(ProjectDir)/Content/SteamAudio/Runtime/...");
+            }
 
 			AddEngineThirdPartyPrivateStaticDependencies(Target, "libPhonon");
 
@@ -58,12 +62,15 @@ namespace UnrealBuildTool.Rules
 			{
 				PrivateDependencyModuleNames.Add("XAudio2");
 				AddEngineThirdPartyPrivateStaticDependencies(Target, "DX11Audio");
-			}
+                RuntimeDependencies.Add("$(EngineDir)/Binaries/ThirdParty/Phonon/Win32/...");
+            }
 			else if (Target.Platform == UnrealTargetPlatform.Win64)
 			{
-				PrivateDependencyModuleNames.Add("XAudio2");
+                PrivateDependencyModuleNames.Add("XAudio2");
 				AddEngineThirdPartyPrivateStaticDependencies(Target, "DX11Audio");
-				PublicDelayLoadDLLs.Add("GPUUtilities.dll");
+
+                RuntimeDependencies.Add("$(EngineDir)/Binaries/ThirdParty/Phonon/Win64/...");
+                PublicDelayLoadDLLs.Add("GPUUtilities.dll");
 				PublicDelayLoadDLLs.Add("tanrt64.dll");
 				PublicDelayLoadDLLs.Add("embree.dll");
 				PublicDelayLoadDLLs.Add("tbb.dll");
diff --git a/Engine/Plugins/Runtime/Steam/SteamAudio/Source/SteamAudioEditor/Private/IndirectBaker.cpp b/Engine/Plugins/Runtime/Steam/SteamAudio/Source/SteamAudioEditor/Private/IndirectBaker.cpp
--- a/Engine/Plugins/Runtime/Steam/SteamAudio/Source/SteamAudioEditor/Private/IndirectBaker.cpp
+++ b/Engine/Plugins/Runtime/Steam/SteamAudio/Source/SteamAudioEditor/Private/IndirectBaker.cpp
@@ -129,7 +129,7 @@ namespace SteamAudio
 			GBakeTickable->SetDisplayText(NSLOCTEXT("SteamAudio", "LoadingScene", "Loading scene..."));
 
 			// Load the scene
-			if (!LoadSceneFromDisk(World, ComputeDevice, SimulationSettings, &PhononScene, PhononSceneInfo))
+			if (!LoadSceneFromDisk(World, ComputeDevice, SimulationSettings, &PhononScene, PhononSceneInfo, nullptr))
 			{
 				// If we can't find the scene, then presumably they haven't generated probes either, so just exit
 				UE_LOG(LogSteamAudioEditor, Error, TEXT("Unable to create Phonon environment: .phononscene not found. Be sure to export the scene."));
diff --git a/Engine/Plugins/Runtime/Steam/SteamAudio/Source/SteamAudioEditor/Private/PhononProbeVolumeDetails.cpp b/Engine/Plugins/Runtime/Steam/SteamAudio/Source/SteamAudioEditor/Private/PhononProbeVolumeDetails.cpp
--- a/Engine/Plugins/Runtime/Steam/SteamAudio/Source/SteamAudioEditor/Private/PhononProbeVolumeDetails.cpp
+++ b/Engine/Plugins/Runtime/Steam/SteamAudio/Source/SteamAudioEditor/Private/PhononProbeVolumeDetails.cpp
@@ -222,11 +222,11 @@ namespace SteamAudio
 			}
 
 			// Attempt to load from disk, otherwise export
-			if (!LoadSceneFromDisk(World, ComputeDevice, SimulationSettings, &PhononScene, PhononSceneInfo))
+			if (!LoadSceneFromDisk(World, ComputeDevice, SimulationSettings, &PhononScene, PhononSceneInfo, nullptr))
 			{
 				IPLhandle PhononStaticMesh = nullptr;
 
-				if (!CreateScene(World, &PhononScene, &PhononStaticMesh, PhononSceneInfo.NumTriangles))
+				if (!CreateScene(World, &PhononScene, &PhononStaticMesh, PhononSceneInfo.NumTriangles, PhononSceneInfo.NumDynTriangles, PhononSceneInfo.DynDataSize))
 				{
 					GGenerateProbesTickable->QueueWorkItem(FWorkItem([](FText& DisplayText) {
 						DisplayText = NSLOCTEXT("SteamAudio", "UnableToCreateScene", "Unable to create scene.");
diff --git a/Engine/Plugins/Runtime/Steam/SteamAudio/Source/SteamAudioEditor/Private/SteamAudioEdModeToolkit.cpp b/Engine/Plugins/Runtime/Steam/SteamAudio/Source/SteamAudioEditor/Private/SteamAudioEdModeToolkit.cpp
--- a/Engine/Plugins/Runtime/Steam/SteamAudio/Source/SteamAudioEditor/Private/SteamAudioEdModeToolkit.cpp
+++ b/Engine/Plugins/Runtime/Steam/SteamAudio/Source/SteamAudioEditor/Private/SteamAudioEdModeToolkit.cpp
@@ -142,6 +142,26 @@ namespace SteamAudio
 				]
 			]
 			+ SVerticalBox::Slot()
+			[
+				SNew(SHorizontalBox)
+				+ SHorizontalBox::Slot()
+				.AutoWidth()
+				.Padding(4)
+				[
+					SNew(STextBlock)
+					.Text(NSLOCTEXT("SteamAudio", "DynamicSceneTriangles", "Dynamic Scene Triangles"))
+					.Font(IDetailLayoutBuilder::GetDetailFont())
+				]
+				+ SHorizontalBox::Slot()
+				.AutoWidth()
+				.Padding(4)
+				[
+					SNew(STextBlock)
+					.Text(this, &FSteamAudioEdModeToolkit::GetNumDynSceneTrianglesText)
+					.Font(IDetailLayoutBuilder::GetDetailFont())
+				]
+			]
+			+ SVerticalBox::Slot()
 			[
 				SNew(SHorizontalBox)
 				+ SHorizontalBox::Slot()
@@ -161,6 +181,26 @@ namespace SteamAudio
 					.Font(IDetailLayoutBuilder::GetDetailFont())
 				]
 			]
+			+ SVerticalBox::Slot()
+			[
+				SNew(SHorizontalBox)
+				+ SHorizontalBox::Slot()
+				.AutoWidth()
+				.Padding(4)
+				[
+					SNew(STextBlock)
+					.Text(NSLOCTEXT("SteamAudio", "DynamicSceneDataSize", "Dynamic Scene Data Size"))
+					.Font(IDetailLayoutBuilder::GetDetailFont())
+				]
+				+ SHorizontalBox::Slot()
+				.AutoWidth()
+				.Padding(4)
+				[
+					SNew(STextBlock)
+					.Text(this, &FSteamAudioEdModeToolkit::GetDynSceneDataSizeText)
+				.Font(IDetailLayoutBuilder::GetDetailFont())
+				]
+			]
 		];
 
 		FModeToolkit::Init(InitToolkitHost);
@@ -252,7 +292,7 @@ namespace SteamAudio
 			IPLhandle PhononStaticMesh = nullptr;
 			FPhononSceneInfo PhononSceneInfo;
 
-			if (!CreateScene(World, &PhononScene, &PhononStaticMesh, PhononSceneInfo.NumTriangles))
+			if (!CreateScene(World, &PhononScene, &PhononStaticMesh, PhononSceneInfo.NumTriangles, PhononSceneInfo.NumDynTriangles, PhononSceneInfo.DynDataSize, true))
 			{
 				GEdModeTickable->QueueWorkItem(FWorkItem([](FText& DisplayText) {
 					DisplayText = NSLOCTEXT("SteamAudio", "ErrorExportingOBJ", "Error exporting OBJ.");
@@ -295,7 +335,7 @@ namespace SteamAudio
 			IPLhandle PhononStaticMesh = nullptr;
 			FPhononSceneInfo PhononSceneInfo;
 
-			if (!CreateScene(World, &PhononScene, &PhononStaticMesh, PhononSceneInfo.NumTriangles))
+			if (!CreateScene(World, &PhononScene, &PhononStaticMesh, PhononSceneInfo.NumTriangles, PhononSceneInfo.NumDynTriangles, PhononSceneInfo.DynDataSize))
 			{
 				GEdModeTickable->QueueWorkItem(FWorkItem([](FText& DisplayText) {
 					DisplayText = NSLOCTEXT("SteamAudio", "ErrorExportingScene", "Error exporting scene.");
@@ -351,6 +391,18 @@ namespace SteamAudio
 		return FText::FromString("No .phononscene found");
 	}
 
+	FText FSteamAudioEdModeToolkit::GetNumDynSceneTrianglesText() const
+	{
+		FSteamAudioEditorModule* Module = &FModuleManager::GetModuleChecked<FSteamAudioEditorModule>("SteamAudioEditor");
+		if (Module != nullptr)
+		{
+			auto PhononSceneInfo = Module->GetCurrentPhononSceneInfo();
+			return FText::AsNumber(PhononSceneInfo.NumDynTriangles);
+		}
+
+		return FText::FromString("No .phononscene found");
+	}
+
 	FText FSteamAudioEdModeToolkit::GetSceneDataSizeText() const
 	{
 		FSteamAudioEditorModule* Module = &FModuleManager::GetModuleChecked<FSteamAudioEditorModule>("SteamAudioEditor");
@@ -362,4 +414,16 @@ namespace SteamAudio
 
 		return FText::FromString("No .phononscene found");
 	}
+
+	FText FSteamAudioEdModeToolkit::GetDynSceneDataSizeText() const
+	{
+		FSteamAudioEditorModule* Module = &FModuleManager::GetModuleChecked<FSteamAudioEditorModule>("SteamAudioEditor");
+		if (Module != nullptr)
+		{
+			auto PhononSceneInfo = Module->GetCurrentPhononSceneInfo();
+			return FText::AsMemory(PhononSceneInfo.DynDataSize);
+		}
+
+		return FText::FromString("No .phononscene found");
+	}
 }
diff --git a/Engine/Plugins/Runtime/Steam/SteamAudio/Source/SteamAudioEditor/Private/SteamAudioEdModeToolkit.h b/Engine/Plugins/Runtime/Steam/SteamAudio/Source/SteamAudioEditor/Private/SteamAudioEdModeToolkit.h
--- a/Engine/Plugins/Runtime/Steam/SteamAudio/Source/SteamAudioEditor/Private/SteamAudioEdModeToolkit.h
+++ b/Engine/Plugins/Runtime/Steam/SteamAudio/Source/SteamAudioEditor/Private/SteamAudioEdModeToolkit.h
@@ -46,7 +46,9 @@ namespace SteamAudio
 		bool IsExportSceneEnabled() const;
 
 		FText GetNumSceneTrianglesText() const;
+		FText GetNumDynSceneTrianglesText() const;
 		FText GetSceneDataSizeText() const;
+		FText GetDynSceneDataSizeText() const;
 
 	private:
 		TSharedPtr<SWidget> ToolkitWidget;
diff --git a/Engine/Plugins/Runtime/Steam/SteamAudio/Source/SteamAudioEditor/Private/SteamAudioEditorModule.cpp b/Engine/Plugins/Runtime/Steam/SteamAudio/Source/SteamAudioEditor/Private/SteamAudioEditorModule.cpp
--- a/Engine/Plugins/Runtime/Steam/SteamAudio/Source/SteamAudioEditor/Private/SteamAudioEditorModule.cpp
+++ b/Engine/Plugins/Runtime/Steam/SteamAudio/Source/SteamAudioEditor/Private/SteamAudioEditorModule.cpp
@@ -85,11 +85,21 @@ namespace SteamAudio
 			PlatformFile.CreateDirectory(*RuntimePath);
 		}
 
+		if (!PlatformFile.DirectoryExists(*DynamicRuntimePath))
+		{
+			PlatformFile.CreateDirectory(*DynamicRuntimePath);
+		}
+
 		if (!PlatformFile.DirectoryExists(*EditorOnlyPath))
 		{
 			PlatformFile.CreateDirectory(*EditorOnlyPath);
 		}
 
+		if (!PlatformFile.DirectoryExists(*DynamicEditorOnlyPath))
+		{
+			PlatformFile.CreateDirectory(*DynamicEditorOnlyPath);
+		}
+
 		// Register detail customizations
 		FPropertyEditorModule& PropertyModule = FModuleManager::LoadModuleChecked<FPropertyEditorModule>("PropertyEditor");
 
diff --git a/Engine/Source/Runtime/Engine/Private/AudioDevice.cpp b/Engine/Source/Runtime/Engine/Private/AudioDevice.cpp
--- a/Engine/Source/Runtime/Engine/Private/AudioDevice.cpp
+++ b/Engine/Source/Runtime/Engine/Private/AudioDevice.cpp
@@ -3006,11 +3006,6 @@ void FAudioDevice::SetListener(UWorld* World, const int32 InViewportIndex, const
 	{
 		ListenerProxies.AddDefaulted(InViewportIndex - ListenerProxies.Num() + 1);
 	}
-	else if (ListenerProxies[InViewportIndex].Transform.Equals(ListenerTransformCopy, KINDA_SMALL_NUMBER))
-	{
-		// Don't need tp update the listener if we're very close to our current transform
-		return;
-	}
 
 	ListenerProxies[InViewportIndex].Transform = ListenerTransformCopy;
 
