Remove test macro added a few weeks ago that is no longer relivant and causing compiler errors in Phonon
#rnx

diff --git a/Engine/Plugins/Runtime/Steam/SteamAudio/Source/SteamAudio/Private/PhononReverb.h b/Engine/Plugins/Runtime/Steam/SteamAudio/Source/SteamAudio/Private/PhononReverb.h
--- a/Engine/Plugins/Runtime/Steam/SteamAudio/Source/SteamAudio/Private/PhononReverb.h
+++ b/Engine/Plugins/Runtime/Steam/SteamAudio/Source/SteamAudio/Private/PhononReverb.h
@@ -126,7 +126,7 @@ class USubmixEffectReverbPluginPreset : public USoundEffectSubmixPreset
 	GENERATED_BODY()
 
 public:
-	EFFECT_PRESET_METHODS_NO_ASSET_ACTIONS(SubmixEffectReverbPlugin)
+	EFFECT_PRESET_METHODS(SubmixEffectReverbPlugin)
 
 	UPROPERTY(EditAnywhere, Category = SubmixEffectPreset)
 	FSubmixEffectReverbPluginSettings Settings;
diff --git a/Engine/Source/Runtime/Engine/Classes/Sound/SoundEffectBase.h b/Engine/Source/Runtime/Engine/Classes/Sound/SoundEffectBase.h
--- a/Engine/Source/Runtime/Engine/Classes/Sound/SoundEffectBase.h
+++ b/Engine/Source/Runtime/Engine/Classes/Sound/SoundEffectBase.h
@@ -57,11 +57,6 @@
 		FCriticalSection SettingsCritSect; \
 		F##EFFECT_NAME##Settings SettingsCopy; \
 
-
-#define EFFECT_PRESET_METHODS_NO_ASSET_ACTIONS(EFFECT_NAME) \
-		virtual bool HasAssetActions() const override { return false; } \
-		EFFECT_PRESET_METHODS(EFFECT_NAME)
-
 class USoundEffectPreset;
 
 class ENGINE_API FSoundEffectBase
