IWYU Fix, missing create call update
#rnx

diff --git a/Engine/Plugins/Runtime/Steam/SteamAudio/Source/SteamAudio/Private/PhononReverb.cpp b/Engine/Plugins/Runtime/Steam/SteamAudio/Source/SteamAudio/Private/PhononReverb.cpp
--- a/Engine/Plugins/Runtime/Steam/SteamAudio/Source/SteamAudio/Private/PhononReverb.cpp
+++ b/Engine/Plugins/Runtime/Steam/SteamAudio/Source/SteamAudio/Private/PhononReverb.cpp
@@ -410,7 +410,7 @@ namespace SteamAudio
 				ReverbPluginPreset = NewObject<USubmixEffectReverbPluginPreset>(Submix, TEXT("Reverb Plugin Effect Preset"));
 			}
 
-			SubmixEffect = FSoundEffectSubmix::Create(FSoundEffectSubmixInitData(), *ReverbPluginPreset);
+			SubmixEffect = USoundEffectPreset::CreateInstance<FSoundEffectSubmixInitData, FSoundEffectSubmix>(FSoundEffectSubmixInitData(), *ReverbPluginPreset);
 			StaticCastSharedPtr<FSubmixEffectReverbPlugin, FSoundEffectSubmix, ESPMode::ThreadSafe>(SubmixEffect)->SetPhononReverbPlugin(this);
 			SubmixEffect->SetEnabled(true);
 		}
diff --git a/Engine/Source/Runtime/Engine/Classes/Sound/SoundEffectBase.h b/Engine/Source/Runtime/Engine/Classes/Sound/SoundEffectBase.h
--- a/Engine/Source/Runtime/Engine/Classes/Sound/SoundEffectBase.h
+++ b/Engine/Source/Runtime/Engine/Classes/Sound/SoundEffectBase.h
@@ -5,6 +5,7 @@
 #include "AudioResampler.h"
 #include "UObject/ObjectMacros.h"
 #include "UObject/Object.h"
+#include "UObject/WeakObjectPtrTemplates.h"
 #include "HAL/ThreadSafeBool.h"
 #include "Containers/Queue.h"
 #include "Misc/ScopeLock.h"
