PR #8720: Change log verbosity in Steam Audio spatialization plugin to stop spam (Contributed by [REDACTED_PERSON])

diff --git a/Engine/Plugins/Runtime/Steam/SteamAudio/Source/SteamAudio/Private/PhononSpatialization.cpp b/Engine/Plugins/Runtime/Steam/SteamAudio/Source/SteamAudio/Private/PhononSpatialization.cpp
--- a/Engine/Plugins/Runtime/Steam/SteamAudio/Source/SteamAudio/Private/PhononSpatialization.cpp
+++ b/Engine/Plugins/Runtime/Steam/SteamAudio/Source/SteamAudio/Private/PhononSpatialization.cpp
@@ -80,7 +80,7 @@ namespace SteamAudio
 		UPhononSpatializationSourceSettings* SpatializationSettings = static_cast<UPhononSpatializationSourceSettings*>(InSettings);
 		FBinauralSource& BinauralSource = BinauralSources[SourceId];
 
-		UE_LOG(LogSteamAudio, Log, TEXT("Creating spatialization effect."));
+		UE_LOG(LogSteamAudio, Verbose, TEXT("Creating spatialization effect."));
 
 		if (SpatializationSettings)
 		{
@@ -109,7 +109,7 @@ namespace SteamAudio
 	{
 		FBinauralSource& BinauralSource = BinauralSources[SourceId];
 
-		UE_LOG(LogSteamAudio, Log, TEXT("Destroying spatialization effect."));
+		UE_LOG(LogSteamAudio, Verbose, TEXT("Destroying spatialization effect."));
 
 		switch (BinauralSource.SpatializationMethod)
 		{
@@ -177,4 +177,4 @@ namespace SteamAudio
 			iplDestroyPanningEffect(&PanningEffect);
 		}
 	}
-}
\ No newline at end of file
+}
