Comment fix

diff --git a/Engine/Plugins/Runtime/Steam/SteamAudio/Source/SteamAudio/Private/PhononSpatialization.cpp b/Engine/Plugins/Runtime/Steam/SteamAudio/Source/SteamAudio/Private/PhononSpatialization.cpp
--- a/Engine/Plugins/Runtime/Steam/SteamAudio/Source/SteamAudio/Private/PhononSpatialization.cpp
+++ b/Engine/Plugins/Runtime/Steam/SteamAudio/Source/SteamAudio/Private/PhononSpatialization.cpp
@@ -129,7 +129,7 @@ namespace SteamAudio
 		BinauralSource.OutBuffer.interleavedBuffer = OutputData.AudioBuffer.GetData();
 
 		// Workaround. The directions passed to spatializer is not consistent with the
-		// coordinate system of UE4, therefore special tranformation is performed here.
+		// coordinate system of UE, therefore special tranformation is performed here.
 		// Review this change if further changes are made to the direction passed to the
 		// spatializer.
 		IPLVector3 RelativeDirection;
