Updating comment

diff --git a/Engine/Plugins/Runtime/Steam/SteamAudio/Source/SteamAudio/Private/PhononScene.cpp b/Engine/Plugins/Runtime/Steam/SteamAudio/Source/SteamAudio/Private/PhononScene.cpp
--- a/Engine/Plugins/Runtime/Steam/SteamAudio/Source/SteamAudio/Private/PhononScene.cpp
+++ b/Engine/Plugins/Runtime/Steam/SteamAudio/Source/SteamAudio/Private/PhononScene.cpp
@@ -597,7 +597,7 @@ namespace SteamAudio
 	//==============================================================================================================================================
 
 	/**
-	 * Populates VertexArray with the given mesh's vertices. Converts from UE4 coords to Phonon coords. Returns the number of vertices added.
+	 * Populates VertexArray with the given mesh's vertices. Converts from UE coords to Phonon coords. Returns the number of vertices added.
 	 */
 	static uint32 GetMeshVerts(UStaticMeshComponent* StaticMeshComponent, TArray<IPLVector3>& VertexArray, bool bGetRelativeVertexPos /*= true*/)
 	{
