Removing UE4 references in comments

diff --git a/Engine/Plugins/Runtime/Steam/SteamAudio/Source/SteamAudio/Private/PhononCommon.h b/Engine/Plugins/Runtime/Steam/SteamAudio/Source/SteamAudio/Private/PhononCommon.h
--- a/Engine/Plugins/Runtime/Steam/SteamAudio/Source/SteamAudio/Private/PhononCommon.h
+++ b/Engine/Plugins/Runtime/Steam/SteamAudio/Source/SteamAudio/Private/PhononCommon.h
@@ -134,7 +134,7 @@ namespace SteamAudio
 	extern FString STEAMAUDIO_API EditorOnlyPath;
 	extern FString STEAMAUDIO_API DynamicEditorOnlyPath;
 
-	/** Functions to convert to/from Phonon/UE4 and IPLVector3/FVector. */
+	/** Functions to convert to/from Phonon/UE and IPLVector3/FVector. */
 	IPLVector3 STEAMAUDIO_API IPLVector3FromFVector(const FVector& Coords);
 	FVector STEAMAUDIO_API FVectorFromIPLVector3(const IPLVector3& Coords);
 	FVector STEAMAUDIO_API UnrealToPhononFVector(const FVector& Coords, const bool bScale = true);
@@ -142,17 +142,17 @@ namespace SteamAudio
 	FVector STEAMAUDIO_API PhononToUnrealFVector(const FVector& Coords, const bool bScale = true);
 	IPLVector3 STEAMAUDIO_API PhononToUnrealIPLVector3(const FVector& Coords, const bool bScale = true);
 
-	/** Given a UE4 transform, produces a corresponding 4x4 transformation matrix. */
+	/** Given a UE transform, produces a corresponding 4x4 transformation matrix. */
 	void STEAMAUDIO_API GetMatrixForTransform(const FTransform& Transform, float* OutMatrix, bool bOutputRowMajor = false, bool bApplyScale = true);
 
 	/** Given a matrix (16 contiguous floats), return a corresponding IPL 4x4 transformation matrix. */
 	IPLMatrix4x4 STEAMAUDIO_API GetIPLMatrix(float* InMatrix);
 
-	/** Phonon raytracer callback that routes ClosestHit queries to the UE4 raytracer. */
+	/** Phonon raytracer callback that routes ClosestHit queries to the UE raytracer. */
 	void STEAMAUDIO_API ClosestHit(const IPLfloat32* Origin, const IPLfloat32* Direction, const IPLfloat32 MinDistance,
 		const IPLfloat32 MaxDistance, IPLfloat32* HitDistance, IPLfloat32* HitNormal, IPLMaterial** HitMaterial, IPLvoid* UserData);
 
-	/** Phonon raytracer callback that routes AnyHit queries to the UE4 raytracer. */
+	/** Phonon raytracer callback that routes AnyHit queries to the UE raytracer. */
 	void STEAMAUDIO_API AnyHit(const IPLfloat32* Origin, const IPLfloat32* Direction, const IPLfloat32 MinDistance, const IPLfloat32 MaxDistance,
 		IPLint32* HitExists, IPLvoid* UserData);
 
