Fixing nonunity errors.

diff --git a/Engine/Plugins/Runtime/Steam/SteamAudio/Source/SteamAudio/Private/PhononPluginManager.cpp b/Engine/Plugins/Runtime/Steam/SteamAudio/Source/SteamAudio/Private/PhononPluginManager.cpp
--- a/Engine/Plugins/Runtime/Steam/SteamAudio/Source/SteamAudio/Private/PhononPluginManager.cpp
+++ b/Engine/Plugins/Runtime/Steam/SteamAudio/Source/SteamAudio/Private/PhononPluginManager.cpp
@@ -8,6 +8,7 @@
 #include "HAL/FileManager.h"
 #include "DrawDebugHelpers.h"
 #include "Engine/World.h"
+#include "Misc/Paths.h"
 
 namespace SteamAudio
 {
