Removing dependencies on SteamAudio/Private by moving some headers to SteamAudio/Internal.

diff --git a/Engine/Plugins/Runtime/Steam/SteamAudio/Source/SteamAudio/Internal/PhononCommon.h b/Engine/Plugins/Runtime/Steam/SteamAudio/Source/SteamAudio/Internal/PhononCommon.h
new file mode 100644
--- /dev/null
+++ b/Engine/Plugins/Runtime/Steam/SteamAudio/Source/SteamAudio/Internal/PhononCommon.h
@@ -0,0 +1,172 @@
+//
+// Copyright (C) Valve Corporation. All rights reserved.
+//
+
+#pragma once
+
+#include <phonon.h>
+#include "CoreMinimal.h"
+#include "UObject/ObjectMacros.h"
+
+#include "PhononCommon.generated.h"
+
+class UPhononGeometryComponent;
+
+DECLARE_LOG_CATEGORY_EXTERN(LogSteamAudio, Log, All);
+
+UENUM(BlueprintType)
+enum class EQualitySettings : uint8
+{
+	LOW				UMETA(DisplayName = "Low"),
+	MEDIUM			UMETA(DisplayName = "Medium"),
+	HIGH			UMETA(DisplayName = "High"),
+	CUSTOM			UMETA(DisplayName = "Custom")
+};
+
+UENUM(BlueprintType)
+enum class EIplSpatializationMethod : uint8
+{
+	// Classic 2D panning - fast.
+	PANNING			UMETA(DisplayName = "Panning"),
+	// Full 3D audio processing with HRTF.
+	HRTF			UMETA(DisplayName = "HRTF")
+};
+
+UENUM(BlueprintType)
+enum class EIplHrtfInterpolationMethod : uint8
+{
+	// Uses a nearest neighbor lookup - fast.
+	NEAREST			UMETA(DisplayName = "Nearest"),
+	// Bilinearly interpolates the HRTF before processing. Slower, but can result in a smoother sound as the listener rotates.
+	BILINEAR		UMETA(DisplayName = "Bilinear")
+};
+
+UENUM(BlueprintType)
+enum class EIplDirectOcclusionMethod : uint8
+{
+	// Binary visible or not test. Adjusts direct volume accordingly.
+	RAYCAST			UMETA(DisplayName = "Raycast"),
+	// Treats the source as a sphere instead of a point. Smoothly ramps up volume as source becomes visible to listener.
+	VOLUMETRIC		UMETA(DisplayName = "Partial")
+};
+
+UENUM(BlueprintType)
+enum class EIplDirectOcclusionMode : uint8
+{
+	// Do not perform any occlusion checks.
+	NONE										UMETA(DisplayName = "None"),
+	// Perform occlusion checks but do not model transmission.
+	DIRECTOCCLUSION_NOTRANSMISSION				UMETA(DisplayName = "Direct Occlusion, No Transmission"),
+	// Perform occlusion checks and model transmission; occluded sound will be scaled by a frequency-independent attenuation value.
+	DIRECTOCCLUSION_TRANSMISSIONBYVOLUME		UMETA(DisplayName = "Direct Occlusion, Frequency-Independent Transmission"),
+	// Perform occlusion checks and model transmission; occluded sound will be rendered with a frequency-dependent transmission filter.
+	DIRECTOCCLUSION_TRANSMISSIONBYFREQUENCY		UMETA(DisplayName = "Direct Occlusion, Frequency-Dependent Transmission")
+};
+
+UENUM(BlueprintType)
+enum class EIplSimulationType : uint8
+{
+	// Simulate indirect sound at run time.
+	REALTIME		UMETA(DisplayName = "Real-Time"),
+	// Precompute indirect sound.
+	BAKED			UMETA(DisplayName = "Baked"),
+	// Do not simulate indirect sound.
+	DISABLED		UMETA(DisplayName = "Disabled")
+};
+
+UENUM(BlueprintType)
+enum class EIplConvolutionType : uint8
+{
+	// Default CPU convolution renderer.
+	PHONON			UMETA(DisplayName = "Phonon"),
+	// AMD TrueAudio Next GPU convolution renderer.
+	TRUEAUDIONEXT	UMETA(DisplayName = "AMD TrueAudio Next")
+};
+
+UENUM(BlueprintType)
+enum class EIplRayTracerType : uint8
+{
+	// Default Ray Tracer
+	PHONON			UMETA(DisplayName = "Phonon"),
+	// Intel Embree Ray Tracer
+	EMBREE			UMETA(DisplayName = "Intel Embree"),
+	// AMD Radeon Rays ray tracer, implemented in OpenCL for both CPU and GPU.
+	RADEONRAYS		UMETA(DisplayName = "AMD Radeon Rays")
+};
+
+UENUM(BlueprintType)
+enum class EIplAudioEngine : uint8
+{
+	// Native Unreal audio engine.
+	UNREAL			UMETA(DisplayName = "Unreal")
+};
+
+namespace SteamAudio
+{
+	struct FSimulationQualitySettings
+	{
+		int32 Bounces;
+		int32 Rays;
+		int32 SecondaryRays;
+	};
+
+	struct FDynamicGeometryMap
+	{
+		FString ID;
+		UPhononGeometryComponent* DynamicGeometryComponent;
+		IPLhandle DynamicScene;
+		IPLhandle DynamicGeometry;
+	};
+
+	/** 1 Unreal Unit = 1cm, 1 Phonon Unit = 1m */
+	const float SCALEFACTOR = 0.01f;
+
+	extern TMap<EQualitySettings, FSimulationQualitySettings> RealtimeSimulationQualityPresets;
+	extern TMap<EQualitySettings, FSimulationQualitySettings> BakedSimulationQualityPresets;
+
+	void* UnrealAlloc(const size_t size, const size_t alignment);
+	void UnrealFree(void* ptr);
+	void UnrealLog(char* msg);
+
+	extern IPLhandle STEAMAUDIO_API GlobalContext;
+
+	extern FString STEAMAUDIO_API BasePath;
+	extern FString STEAMAUDIO_API RuntimePath;
+	extern FString STEAMAUDIO_API DynamicRuntimePath;
+	extern FString STEAMAUDIO_API EditorOnlyPath;
+	extern FString STEAMAUDIO_API DynamicEditorOnlyPath;
+
+	/** Functions to convert to/from Phonon/UE and IPLVector3/FVector. */
+	IPLVector3 STEAMAUDIO_API IPLVector3FromFVector(const FVector& Coords);
+	FVector STEAMAUDIO_API FVectorFromIPLVector3(const IPLVector3& Coords);
+	FVector STEAMAUDIO_API UnrealToPhononFVector(const FVector& Coords, const bool bScale = true);
+	IPLVector3 STEAMAUDIO_API UnrealToPhononIPLVector3(const FVector& Coords, const bool bScale = true);
+	FVector STEAMAUDIO_API PhononToUnrealFVector(const FVector& Coords, const bool bScale = true);
+	IPLVector3 STEAMAUDIO_API PhononToUnrealIPLVector3(const FVector& Coords, const bool bScale = true);
+
+	/** Given a UE transform, produces a corresponding 4x4 transformation matrix. */
+	void STEAMAUDIO_API GetMatrixForTransform(const FTransform& Transform, float* OutMatrix, bool bOutputRowMajor = false, bool bApplyScale = true);
+
+	/** Given a matrix (16 contiguous floats), return a corresponding IPL 4x4 transformation matrix. */
+	IPLMatrix4x4 STEAMAUDIO_API GetIPLMatrix(float* InMatrix);
+
+	/** Phonon raytracer callback that routes ClosestHit queries to the UE raytracer. */
+	void STEAMAUDIO_API ClosestHit(const IPLfloat32* Origin, const IPLfloat32* Direction, const IPLfloat32 MinDistance,
+		const IPLfloat32 MaxDistance, IPLfloat32* HitDistance, IPLfloat32* HitNormal, IPLMaterial** HitMaterial, IPLvoid* UserData);
+
+	/** Phonon raytracer callback that routes AnyHit queries to the UE raytracer. */
+	void STEAMAUDIO_API AnyHit(const IPLfloat32* Origin, const IPLfloat32* Direction, const IPLfloat32 MinDistance, const IPLfloat32 MaxDistance,
+		IPLint32* HitExists, IPLvoid* UserData);
+
+	/** Strips PIE prefix from map name for use in-editor. */
+	FString STEAMAUDIO_API StrippedMapName(const FString& MapName);
+
+	/** Attempts to loads the specified DLL, performing some basic error checking. Returns handle to DLL or nullptr if failed to load.*/
+	void* LoadDll(const FString& DllFile, bool bErrorOnLoadFailure);
+
+	/** Error logs non-successful statuses. */
+	void LogSteamAudioStatus(const IPLerror Status);
+
+	/** Converts IPL statuses to a readable FString. */
+	FString StatusToFString(const IPLerror Status);
+}
diff --git a/Engine/Plugins/Runtime/Steam/SteamAudio/Source/SteamAudio/Internal/PhononMaterial.h b/Engine/Plugins/Runtime/Steam/SteamAudio/Source/SteamAudio/Internal/PhononMaterial.h
new file mode 100644
--- /dev/null
+++ b/Engine/Plugins/Runtime/Steam/SteamAudio/Source/SteamAudio/Internal/PhononMaterial.h
@@ -0,0 +1,33 @@
+//
+// Copyright (C) Valve Corporation. All rights reserved.
+//
+
+#pragma once
+
+#include "CoreMinimal.h"
+#include "UObject/ObjectMacros.h"
+#include "phonon.h"
+
+#include "PhononMaterial.generated.h"
+
+UENUM()
+enum class EPhononMaterial : uint8
+{
+	GENERIC = 0 UMETA(DisplayName = "Generic"),
+	BRICK = 1 UMETA(DisplayName = "Brick"),
+	CONCRETE = 2 UMETA(DisplayName = "Concrete"),
+	CERAMIC = 3 UMETA(DisplayName = "Ceramic"),
+	GRAVEL = 4 UMETA(DisplayName = "Gravel"),
+	CARPET = 5 UMETA(DisplayName = "Carpet"),
+	GLASS = 6 UMETA(DisplayName = "Glass"),
+	PLASTER = 7 UMETA(DisplayName = "Plaster"),
+	WOOD = 8 UMETA(DisplayName = "Wood"),
+	METAL = 9 UMETA(DisplayName = "Metal"),
+	ROCK = 10 UMETA(DisplayName = "Rock"),
+	CUSTOM = 11 UMETA(DisplayName = "Custom")
+};
+
+namespace SteamAudio
+{
+	extern TMap<EPhononMaterial, IPLMaterial> MaterialPresets;
+}
\ No newline at end of file
diff --git a/Engine/Plugins/Runtime/Steam/SteamAudio/Source/SteamAudio/Internal/PhononProbeComponent.h b/Engine/Plugins/Runtime/Steam/SteamAudio/Source/SteamAudio/Internal/PhononProbeComponent.h
new file mode 100644
--- /dev/null
+++ b/Engine/Plugins/Runtime/Steam/SteamAudio/Source/SteamAudio/Internal/PhononProbeComponent.h
@@ -0,0 +1,20 @@
+//
+// Copyright (C) Valve Corporation. All rights reserved.
+//
+
+#pragma once
+
+#include "Components/SceneComponent.h"
+#include "PhononProbeComponent.generated.h"
+
+UCLASS(ClassGroup = (Audio), HideCategories = (Activation, Collision))
+class STEAMAUDIO_API UPhononProbeComponent : public USceneComponent
+{
+	GENERATED_BODY()
+
+public:
+	UPROPERTY()
+	TArray<FVector> ProbeLocations;
+
+	FCriticalSection ProbeLocationsCriticalSection;
+};
diff --git a/Engine/Plugins/Runtime/Steam/SteamAudio/Source/SteamAudio/Internal/PhononProbeVolume.h b/Engine/Plugins/Runtime/Steam/SteamAudio/Source/SteamAudio/Internal/PhononProbeVolume.h
new file mode 100644
--- /dev/null
+++ b/Engine/Plugins/Runtime/Steam/SteamAudio/Source/SteamAudio/Internal/PhononProbeVolume.h
@@ -0,0 +1,136 @@
+//
+// Copyright (C) Valve Corporation. All rights reserved.
+//
+
+#pragma once
+
+#include "PhononCommon.h"
+#include "GameFramework/Volume.h"
+#include "PhononProbeVolume.generated.h"
+
+UENUM()
+enum class EPhononProbePlacementStrategy : uint8
+{
+	// Places a single probe at the centroid of the volume.
+	CENTROID = 0 UMETA(DisplayName = "Centroid"),
+	// Places uniformly spaced probes along the floor at a specified height.
+	UNIFORM_FLOOR = 1 UMETA(DisplayName = "Uniform Floor")
+};
+
+UENUM()
+enum class EPhononProbeMobility : uint8
+{
+	// Static probes remain fixed at runtime.
+	STATIC = 0 UMETA(DisplayName = "Static"),
+	// Dynamic probes inherit this volume's offset at runtime.
+	DYNAMIC = 1 UMETA(DisplayName = "Dynamic")
+};
+
+USTRUCT()
+struct FBakedDataInfo
+{
+	GENERATED_BODY()
+
+	UPROPERTY()
+	FName Name;
+
+	UPROPERTY()
+	int32 Size;
+};
+
+inline bool operator==(const FBakedDataInfo& Lhs, const FBakedDataInfo& Rhs)
+{
+	return Lhs.Name == Rhs.Name && Lhs.Size == Rhs.Size;
+}
+
+inline bool operator<(const FBakedDataInfo& Lhs, const FBakedDataInfo& Rhs)
+{
+	return Lhs.Name.LexicalLess(Rhs.Name);
+}
+
+/**
+ * Phonon Probe volumes generate a set of probes at which acoustic information will be sampled
+ * at bake time.
+ */
+UCLASS(HideCategories = (Actor, Advanced, Attachment, Collision), meta = (BlueprintSpawnableComponent))
+class STEAMAUDIO_API APhononProbeVolume : public AVolume
+{
+	GENERATED_UCLASS_BODY()
+
+public:
+
+#if WITH_EDITOR
+
+	/** Creates a probe box with the current probe placement parameters and writes it to disk. Updates ProbeDataSize for display. */
+	void PlaceProbes(IPLhandle PhononScene, IPLProbePlacementProgressCallback ProbePlacementCallback, TArray<IPLSphere>& ProbeSpheres);
+
+	/** Writes probe box out to disk. Creates a new probe batch and writes it to disk. Updates ProbeDataSize for display. */
+	void UpdateProbeData(IPLhandle ProbeBox);
+
+	/** Writes a given probe box out to the EditorOnly folder. Returns file size, or 0 if no file was written. */
+	int32 SaveProbeBoxToDisk(IPLhandle ProbeBox);
+
+	/** Writes a given probe batch out to the Runtime folder. Returns file size, or 0 if no file was written. */
+	int32 SaveProbeBatchToDisk(IPLhandle ProbeBatch);
+
+	virtual bool CanEditChange(const FProperty* InProperty) const override;
+
+#endif
+
+	/** Returns size of probe box file on disk, or 0 if none exists. */
+	int32 GetProbeBoxDataSize() const;
+
+	/** Returns size of probe batch file on disk, or 0 if none exists. */
+	int32 GetProbeBatchDataSize() const;
+
+	/** Returns data size for the specified source from member data. */
+	int32 GetDataSizeForSource(const FName& UniqueIdentifier) const;
+
+	/** Composes a probe box filename for this actor. TODO: ensure uniqueness. */
+	FString GetProbeBoxFileName() const;
+
+	/** Composes a probe batch filename for this actor. TODO: ensure uniqueness. */
+	FString GetProbeBatchFileName() const;
+
+	bool LoadProbeBoxFromDisk(IPLhandle* ProbeBox) const;
+
+	bool LoadProbeBatchFromDisk(IPLhandle* ProbeBatch) const;
+
+	class UPhononProbeComponent* GetPhononProbeComponent() const { return PhononProbeComponent; }
+
+	// Method by which probes are placed within the volume.
+	UPROPERTY(EditAnywhere, Category = ProbeGeneration)
+	EPhononProbePlacementStrategy PlacementStrategy;
+
+	// How far apart to place probes.
+	UPROPERTY(EditAnywhere, Category = ProbeGeneration, meta = (ClampMin = "0.0", ClampMax = "5000.0", UIMin = "0.0", UIMax = "5000.0"))
+	float HorizontalSpacing;
+
+	// How high above the floor to place probes.
+	UPROPERTY(EditAnywhere, Category = ProbeGeneration, meta = (ClampMin = "0.0", ClampMax = "5000.0", UIMin = "0.0", UIMax = "5000.0"))
+	float HeightAboveFloor;
+
+	// Number of probes contained in this probe volume.
+	UPROPERTY(VisibleAnywhere, Category = ProbeVolumeStatistics, meta = (DisplayName = "Probe Points"))
+	int32 NumProbes;
+
+	// Size of probe data in bytes.
+	UPROPERTY(VisibleAnywhere, Category = ProbeVolumeStatistics, meta = (DisplayName = "Probe Data Size"))
+	int32 ProbeDataSize;
+
+	// Useful information for each baked source.
+	UPROPERTY(VisibleAnywhere, Category = ProbeVolumeStatistics, meta = (DisplayName = "Detailed Statistics"))
+	TArray<FBakedDataInfo> BakedDataInfo;
+
+	// Component used for visualization.
+	UPROPERTY()
+	TObjectPtr<class UPhononProbeComponent> PhononProbeComponent;
+
+	// Current filename where probe box data is stored. Used to maintain connection if volume is renamed.
+	UPROPERTY()
+	FString ProbeBoxFileName;
+
+	// Current filename where probe batch data is stored. Used to maintain connection if volume is renamed.
+	UPROPERTY()
+	FString ProbeBatchFileName;
+};
diff --git a/Engine/Plugins/Runtime/Steam/SteamAudio/Source/SteamAudio/Internal/PhononScene.h b/Engine/Plugins/Runtime/Steam/SteamAudio/Source/SteamAudio/Internal/PhononScene.h
new file mode 100644
--- /dev/null
+++ b/Engine/Plugins/Runtime/Steam/SteamAudio/Source/SteamAudio/Internal/PhononScene.h
@@ -0,0 +1,65 @@
+//
+// Copyright (C) Valve Corporation. All rights reserved.
+//
+
+#pragma once
+
+#include "phonon.h"
+
+#include "Async/Async.h"
+#include "CoreMinimal.h"
+#include "UObject/ObjectMacros.h"
+#include "Engine/StaticMeshActor.h"
+
+class AActor;
+class AStaticMeshActor;
+class UWorld;
+
+namespace SteamAudio
+{
+	struct FPhononSceneInfo
+	{
+		FPhononSceneInfo()
+			: NumTriangles(0)
+			, NumDynTriangles(0)
+			, DataSize(0)
+			, DynDataSize(0)
+		{}
+
+		uint32 NumTriangles;
+		uint32 NumDynTriangles;
+		uint32 DataSize;
+		uint32 DynDataSize;
+	};
+
+	bool STEAMAUDIO_API LoadSceneFromDisk(UWorld* World, IPLhandle ComputeDevice, const IPLSimulationSettings& SimulationSettings,
+		IPLhandle* PhononScene, FPhononSceneInfo& PhononSceneInfo, FString* CurrentLevelName);
+
+	bool STEAMAUDIO_API LoadSceneInfoFromDisk(UWorld* World, FPhononSceneInfo& PhononSceneInfo);
+
+	bool STEAMAUDIO_API LoadDynamicSceneFromDisk(UWorld* World, IPLhandle ComputeDevice, const IPLSimulationSettings& SimulationSettings,
+		IPLhandle* PhononScene, FPhononSceneInfo& PhononSceneInfo, FString DynamicSceneID);
+
+	bool STEAMAUDIO_API LoadDynamicScenesInfoFromDisk(UWorld* World, FPhononSceneInfo& PhononSceneInfo, FString DynamicSceneID);
+
+#if WITH_EDITOR
+
+	/** Creates and populates a Phonon scene. Gathers Phonon Geometry asynchronously on the game thread. */
+	bool STEAMAUDIO_API CreateScene(UWorld* World, IPLhandle* PhononScene, IPLhandle* PhononStaticMesh, uint32& NumSceneTriangles, uint32& NumDynSceneTriangles, uint32& DynDataSize, bool CreateDynObjFiles = false);
+	void STEAMAUDIO_API CreateDynamicScene(UWorld* World, IPLSimulationSettings SimulationSettings, TArray<IPLMaterial> SceneMaterials, uint32& NumSceneTriangles, uint32& DataSize, bool CreateObjFile = false);
+	bool STEAMAUDIO_API SaveFinalizedSceneToDisk(UWorld* World, IPLhandle PhononScene, const FPhononSceneInfo& PhononSceneInfo);
+	bool STEAMAUDIO_API SaveDynamicSceneToDisk(UWorld* World, IPLhandle PhononScene, const FPhononSceneInfo& PhononSceneInfo, FString DynamicSceneID, bool CreateObjFile = false);
+	void STEAMAUDIO_API AddGeometryComponentsToStaticMeshes(UWorld* World);
+	void STEAMAUDIO_API RemoveGeometryComponentsFromStaticMeshes(UWorld* World);
+
+	/** Check a directory if it exists, create it if it doesn't. Return True if the directory exists. */
+	bool STEAMAUDIO_API CheckSceneDirectory(FString DirectoryPath);
+
+	/** Delete any scene files in a given path that start with the world/level name. */
+	void STEAMAUDIO_API DeleteSceneFiles(UWorld* World, FString FilePath);
+
+#endif // WITH_EDITOR
+
+	uint32 GetNumTrianglesForStaticMesh(AStaticMeshActor* StaticMeshActor);
+	uint32 GetNumTrianglesAtRoot(AActor* RootActor);
+}
diff --git a/Engine/Plugins/Runtime/Steam/SteamAudio/Source/SteamAudio/Internal/PhononSourceComponent.h b/Engine/Plugins/Runtime/Steam/SteamAudio/Source/SteamAudio/Internal/PhononSourceComponent.h
new file mode 100644
--- /dev/null
+++ b/Engine/Plugins/Runtime/Steam/SteamAudio/Source/SteamAudio/Internal/PhononSourceComponent.h
@@ -0,0 +1,27 @@
+//
+// Copyright (C) Valve Corporation. All rights reserved.
+//
+
+#pragma once
+
+#include "Components/SceneComponent.h"
+#include "PhononSourceComponent.generated.h"
+
+/**
+ * Phonon Source Components can be placed alongside statically positioned Audio Components in order to bake impulse response data
+ * to be applied at runtime.
+ */
+UCLASS(ClassGroup = (Audio), meta = (BlueprintSpawnableComponent), HideCategories = (Activation, Collision, Tags, Rendering, Physics, LOD, Mobility, Cooking, AssetUserData))
+class STEAMAUDIO_API UPhononSourceComponent : public USceneComponent
+{
+	GENERATED_BODY()
+
+public:
+	// Any Phonon probes that lie within the baking radius will be used to produce baked impulse response data for this source location.
+	UPROPERTY(EditAnywhere, Category = Baking)
+	float BakingRadius = 1600.0f;
+
+	// Users must specify a unique identifier for baked data lookup at runtime.
+	UPROPERTY(EditAnywhere, Category = Baking)
+	FName UniqueIdentifier;
+};
\ No newline at end of file
diff --git a/Engine/Plugins/Runtime/Steam/SteamAudio/Source/SteamAudio/Internal/SteamAudioEnvironment.h b/Engine/Plugins/Runtime/Steam/SteamAudio/Source/SteamAudio/Internal/SteamAudioEnvironment.h
new file mode 100644
--- /dev/null
+++ b/Engine/Plugins/Runtime/Steam/SteamAudio/Source/SteamAudio/Internal/SteamAudioEnvironment.h
@@ -0,0 +1,75 @@
+//
+// Copyright (C) Valve Corporation. All rights reserved.
+//
+
+#pragma once
+
+#include "ISteamAudioModule.h"
+#include "phonon.h"
+#include "PhononCommon.h"
+#include "PhononScene.h"
+
+namespace SteamAudio
+{
+	/**
+	 * Maps from UE FName identifiers stored on AudioComponents to a unique integer for the Phonon API.
+	 * The expected number of baked sources for a given scene is quite low (< 1k usually), so this simply
+	 * maps to random integers in (0, INT_MAX).
+	 */
+	class STEAMAUDIO_API FIdentifierMap
+	{
+	public:
+		bool ContainsKey(const FString& Key) const;
+		bool ContainsValue(const int32 Value) const;
+		int32 Add(const FString& Identifier);
+		void Add(const TPair<FString, int32>& Pair);
+		int32 Get(const FString& Identifier) const;
+		FString ToString() const;
+		void FromString(const FString& MappingString);
+		void Empty();
+
+	private:
+		TArray<TPair<FString, int32>> IdentifierMapping;
+	};
+
+	bool STEAMAUDIO_API LoadBakedIdentifierMapFromDisk(UWorld* World, FIdentifierMap& BakedIdentifierMap);
+	bool STEAMAUDIO_API SaveBakedIdentifierMapToDisk(UWorld* World, const FIdentifierMap& BakedIdentifierMap);
+
+	/**
+	 * Handles an instance of the Steam Audio environment and the environmental renderer used by the audio plugins.
+	 */
+	class FEnvironment
+	{
+	public:
+		FEnvironment();
+		~FEnvironment();
+
+		bool Initialize(UWorld* World, FAudioDevice* InAudioDevice, FString* CurrentLevelName);
+		void Shutdown();
+
+		IPLhandle GetScene() { return PhononScene; };
+		IPLhandle GetEnvironment() const { return PhononEnvironment; };
+		IPLhandle GetComputeDevice() const { return ComputeDevice; };
+		IPLhandle GetEnvironmentalRenderer() const { return EnvironmentalRenderer; }
+		const IPLSimulationSettings& GetSimulationSettings() const { return SimulationSettings; }
+		const IPLRenderingSettings& GetRenderingSettings() const { return RenderingSettings; }
+		const FIdentifierMap& GetBakedIdentifierMap() const { return BakedIdentifierMap; }
+		FCriticalSection* GetEnvironmentCriticalSectionHandle() { return &EnvironmentCriticalSection; };
+
+	private:
+		FCriticalSection EnvironmentCriticalSection;
+
+		IPLhandle ComputeDevice;
+		IPLhandle PhononScene;
+		IPLhandle PhononEnvironment;
+		IPLhandle EnvironmentalRenderer;
+		IPLhandle ProbeManager;
+		TArray<IPLhandle> ProbeBatches;
+
+		IPLSimulationSettings SimulationSettings;
+		IPLRenderingSettings RenderingSettings;
+
+		FPhononSceneInfo PhononSceneInfo;
+		FIdentifierMap BakedIdentifierMap;
+	};
+}
\ No newline at end of file
diff --git a/Engine/Plugins/Runtime/Steam/SteamAudio/Source/SteamAudio/Internal/SteamAudioSettings.h b/Engine/Plugins/Runtime/Steam/SteamAudio/Source/SteamAudio/Internal/SteamAudioSettings.h
new file mode 100644
--- /dev/null
+++ b/Engine/Plugins/Runtime/Steam/SteamAudio/Source/SteamAudio/Internal/SteamAudioSettings.h
@@ -0,0 +1,310 @@
+//
+// Copyright (C) Valve Corporation. All rights reserved.
+//
+
+#pragma once
+
+#include "PhononMaterial.h"
+#include "PhononCommon.h"
+#include "SteamAudioSettings.generated.h"
+
+UCLASS(config = Engine, defaultconfig)
+class STEAMAUDIO_API USteamAudioSettings : public UObject
+{
+	GENERATED_BODY()
+
+public:
+
+	USteamAudioSettings();
+
+	IPLMaterial GetDefaultStaticMeshMaterial() const;
+	IPLMaterial GetDefaultBSPMaterial() const;
+	IPLMaterial GetDefaultLandscapeMaterial() const;
+	int32 GetIndirectImpulseResponseOrder() const;
+	float GetIndirectImpulseResponseDuration() const;
+	uint32 GetMaxSources() const;
+	IPLSimulationSettings GetRealtimeSimulationSettings() const;
+	IPLSimulationSettings GetBakedSimulationSettings() const;
+	int32 GetBakingBatchSize() const;
+	float GetFractionComputeUnitsForIRUpdate() const;
+
+#if WITH_EDITOR
+	virtual void PostEditChangeProperty(struct FPropertyChangedEvent& PropertyChangedEvent) override;
+	virtual bool CanEditChange(const FProperty* InProperty) const override;
+#endif
+
+	UPROPERTY(GlobalConfig, EditAnywhere, Category = General, meta = (AllowedClasses = "/Script/Engine.SoundSubmix"))
+	FSoftObjectPath OutputSubmix;
+
+	// Which audio engine to use.
+	UPROPERTY(GlobalConfig, EditAnywhere, Category = General)
+	EIplAudioEngine AudioEngine;
+
+	// Which ray tracer type to use.
+	UPROPERTY(GlobalConfig, EditAnywhere, Category = General)
+	EIplRayTracerType RayTracer;
+
+	// Which convolution renderer to use.
+	UPROPERTY(GlobalConfig, EditAnywhere, Category = General)
+	EIplConvolutionType ConvolutionType;
+
+	//==============================================================================================================================================
+
+	// Whether or not to export BSP geometry.
+	UPROPERTY(GlobalConfig, EditAnywhere, Category = "Scene Export", meta = (DisplayName = "Export BSP Geometry"))
+	bool ExportBSPGeometry;
+
+	// Whether or not to export Landscape geometry.
+	UPROPERTY(GlobalConfig, EditAnywhere, Category = "Scene Export", meta = (DisplayName = "Export Landscape Geometry"))
+	bool ExportLandscapeGeometry;
+
+	//==============================================================================================================================================
+
+	// Preset material settings for Static Mesh actors.
+	UPROPERTY(GlobalConfig, EditAnywhere, Category = "Scene Export|Default Static Mesh Material", meta = (DisplayName = "Material Preset"))
+	EPhononMaterial StaticMeshMaterialPreset;
+
+	// How much this material absorbs low frequency sound.
+	UPROPERTY(GlobalConfig, EditAnywhere, Category = "Scene Export|Default Static Mesh Material", meta = (ClampMin = "0.0", ClampMax = "1.0", 
+		UIMin = "0.0", UIMax = "1.0", DisplayName = "Low Frequency Absorption"))
+	float StaticMeshLowFreqAbsorption;
+
+	// How much this material absorbs mid frequency sound.
+	UPROPERTY(GlobalConfig, EditAnywhere, Category = "Scene Export|Default Static Mesh Material", meta = (ClampMin = "0.0", ClampMax = "1.0", 
+		UIMin = "0.0", UIMax = "1.0", DisplayName = "Mid Frequency Absorption"))
+	float StaticMeshMidFreqAbsorption;
+
+	// How much this material absorbs high frequency sound.
+	UPROPERTY(GlobalConfig, EditAnywhere, Category = "Scene Export|Default Static Mesh Material", meta = (ClampMin = "0.0", ClampMax = "1.0", 
+		UIMin = "0.0", UIMax = "1.0", DisplayName = "High Frequency Absorption"))
+	float StaticMeshHighFreqAbsorption; 
+
+	// How much this material transmits low frequency sound.
+	UPROPERTY(GlobalConfig, EditAnywhere, Category = "Scene Export|Default Static Mesh Material", meta = (ClampMin = "0.0", ClampMax = "1.0",
+		UIMin = "0.0", UIMax = "1.0", DisplayName = "Low Frequency Transmission"))
+	float StaticMeshLowFreqTransmission;
+
+	// How much this material transmits mid frequency sound.
+	UPROPERTY(GlobalConfig, EditAnywhere, Category = "Scene Export|Default Static Mesh Material", meta = (ClampMin = "0.0", ClampMax = "1.0",
+		UIMin = "0.0", UIMax = "1.0", DisplayName = "Mid Frequency Transmission"))
+	float StaticMeshMidFreqTransmission;
+
+	// How much this material transmits high frequency sound.
+	UPROPERTY(GlobalConfig, EditAnywhere, Category = "Scene Export|Default Static Mesh Material", meta = (ClampMin = "0.0", ClampMax = "1.0",
+		UIMin = "0.0", UIMax = "1.0", DisplayName = "High Frequency Transmission"))
+	float StaticMeshHighFreqTransmission;
+
+	// Specifies how "rough" the surface is. Surfaces with a high scattering value randomly reflect sound in all directions;
+	// surfaces with a low scattering value reflect sound in a mirror-like manner.
+	UPROPERTY(GlobalConfig, EditAnywhere, Category = "Scene Export|Default Static Mesh Material", meta = (ClampMin = "0.0", ClampMax = "1.0", 
+		UIMin = "0.0", UIMax = "1.0", DisplayName = "Scattering"))
+	float StaticMeshScattering;
+
+	//==============================================================================================================================================
+
+	// Preset material settings for BSP geometry.
+	UPROPERTY(GlobalConfig, EditAnywhere, Category = "Scene Export|Default BSP Material", meta = (DisplayName = "Material Preset"))
+	EPhononMaterial BSPMaterialPreset;
+
+	// How much this material absorbs low frequency sound.
+	UPROPERTY(GlobalConfig, EditAnywhere, Category = "Scene Export|Default BSP Material", meta = (ClampMin = "0.0", ClampMax = "1.0", 
+		UIMin = "0.0", UIMax = "1.0", DisplayName = "Low Frequency Absorption"))
+	float BSPLowFreqAbsorption;
+
+	// How much this material absorbs mid frequency sound.
+	UPROPERTY(GlobalConfig, EditAnywhere, Category = "Scene Export|Default BSP Material", meta = (ClampMin = "0.0", ClampMax = "1.0", 
+		UIMin = "0.0", UIMax = "1.0", DisplayName = "Mid Frequency Absorption"))
+	float BSPMidFreqAbsorption;
+
+	// How much this material absorbs high frequency sound.
+	UPROPERTY(GlobalConfig, EditAnywhere, Category = "Scene Export|Default BSP Material", meta = (ClampMin = "0.0", ClampMax = "1.0", 
+		UIMin = "0.0", UIMax = "1.0", DisplayName = "High Frequency Absorption"))
+	float BSPHighFreqAbsorption;
+
+	// How much this material transmits low frequency sound.
+	UPROPERTY(GlobalConfig, EditAnywhere, Category = "Scene Export|Default BSP Material", meta = (ClampMin = "0.0", ClampMax = "1.0",
+		UIMin = "0.0", UIMax = "1.0", DisplayName = "Low Frequency Transmission"))
+	float BSPLowFreqTransmission;
+
+	// How much this material transmits mid frequency sound.
+	UPROPERTY(GlobalConfig, EditAnywhere, Category = "Scene Export|Default BSP Material", meta = (ClampMin = "0.0", ClampMax = "1.0",
+		UIMin = "0.0", UIMax = "1.0", DisplayName = "Mid Frequency Transmission"))
+	float BSPMidFreqTransmission;
+
+	// How much this material transmits high frequency sound.
+	UPROPERTY(GlobalConfig, EditAnywhere, Category = "Scene Export|Default BSP Material", meta = (ClampMin = "0.0", ClampMax = "1.0",
+		UIMin = "0.0", UIMax = "1.0", DisplayName = "High Frequency Transmission"))
+	float BSPHighFreqTransmission;
+
+	// Specifies how "rough" the surface is. Surfaces with a high scattering value randomly reflect sound in all directions;
+	// surfaces with a low scattering value reflect sound in a mirror-like manner.
+	UPROPERTY(GlobalConfig, EditAnywhere, Category = "Scene Export|Default BSP Material", meta = (ClampMin = "0.0", ClampMax = "1.0", 
+		UIMin = "0.0", UIMax = "1.0", DisplayName = "Scattering"))
+	float BSPScattering;
+
+	//==============================================================================================================================================
+
+	// Preset material settings for landscape actors.
+	UPROPERTY(GlobalConfig, EditAnywhere, Category = "Scene Export|Default Landscape Material", meta = (DisplayName = "Material Preset"))
+	EPhononMaterial LandscapeMaterialPreset;
+
+	// How much this material absorbs low frequency sound.
+	UPROPERTY(GlobalConfig, EditAnywhere, Category = "Scene Export|Default Landscape Material", meta = (ClampMin = "0.0", ClampMax = "1.0", 
+		UIMin = "0.0", UIMax = "1.0", DisplayName = "Low Frequency Absorption"))
+	float LandscapeLowFreqAbsorption;
+
+	// How much this material absorbs mid frequency sound.
+	UPROPERTY(GlobalConfig, EditAnywhere, Category = "Scene Export|Default Landscape Material", meta = (ClampMin = "0.0", ClampMax = "1.0", 
+		UIMin = "0.0", UIMax = "1.0", DisplayName = "Mid Frequency Absorption"))
+	float LandscapeMidFreqAbsorption;
+
+	// How much this material absorbs high frequency sound.
+	UPROPERTY(GlobalConfig, EditAnywhere, Category = "Scene Export|Default Landscape Material", meta = (ClampMin = "0.0", ClampMax = "1.0", 
+		UIMin = "0.0", UIMax = "1.0", DisplayName = "High Frequency Absorption"))
+	float LandscapeHighFreqAbsorption;
+
+	// How much this material transmits low frequency sound.
+	UPROPERTY(GlobalConfig, EditAnywhere, Category = "Scene Export|Default Landscape Material", meta = (ClampMin = "0.0", ClampMax = "1.0",
+		UIMin = "0.0", UIMax = "1.0", DisplayName = "Low Frequency Transmission"))
+	float LandscapeLowFreqTransmission;
+
+	// How much this material transmits mid frequency sound.
+	UPROPERTY(GlobalConfig, EditAnywhere, Category = "Scene Export|Default Landscape Material", meta = (ClampMin = "0.0", ClampMax = "1.0",
+		UIMin = "0.0", UIMax = "1.0", DisplayName = "Mid Frequency Transmission"))
+	float LandscapeMidFreqTransmission;
+
+	// How much this material transmits high frequency sound.
+	UPROPERTY(GlobalConfig, EditAnywhere, Category = "Scene Export|Default Landscape Material", meta = (ClampMin = "0.0", ClampMax = "1.0",
+		UIMin = "0.0", UIMax = "1.0", DisplayName = "High Frequency Transmission"))
+	float LandscapeHighFreqTransmission;
+
+	// Specifies how "rough" the surface is. Surfaces with a high scattering value randomly reflect sound in all directions;
+	// surfaces with a low scattering value reflect sound in a mirror-like manner.
+	UPROPERTY(GlobalConfig, EditAnywhere, Category = "Scene Export|Default Landscape Material", meta = (ClampMin = "0.0", ClampMax = "1.0", 
+		UIMin = "0.0", UIMax = "1.0", DisplayName = "Scattering"))
+	float LandscapeScattering;
+
+	//==============================================================================================================================================
+
+	// The number of rays to trace from the listener to a source when simulating volumetric occlusion. Increasing this number increases the smoothness of occlusion transitions, but also increases CPU usage and memory consumption.
+	UPROPERTY(GlobalConfig, EditAnywhere, Category = Occlusion, meta = (ClampMin = "4", ClampMax = "1024", UIMin = "4", UIMax = "1024"))
+	int32 OcclusionSampleCount;
+
+	//==============================================================================================================================================
+
+	// How to simulate listener-centric reverb.
+	UPROPERTY(GlobalConfig, EditAnywhere, Category = Reverb, meta = (DisplayName = "Listener-Centric Reverb Simulation"))
+	EIplSimulationType ListenerReverbSimulationType;
+
+	// How much listener-centric reverb should contribute.
+	UPROPERTY(GlobalConfig, EditAnywhere, Category = Reverb, meta = (DisplayName = "Listener-Centric Reverb Contribution", ClampMin = "0.0", ClampMax = "10.0", UIMin = "0.0", UIMax = "10.0"))
+	float ListenerReverbContribution;
+
+	// Output of indirect propagation is stored in ambisonics of this order.
+	UPROPERTY(GlobalConfig, EditAnywhere, Category = Reverb, meta = (ClampMin = "0", ClampMax = "3", UIMin = "0", UIMax = "3",
+		DisplayName = "Ambisonics Order"))
+	int32 IndirectImpulseResponseOrder;
+
+	// Length of impulse response to compute for each sound source.
+	UPROPERTY(GlobalConfig, EditAnywhere, Category = Reverb, meta = (ClampMin = "0.1", ClampMax = "5.0", UIMin = "0.1", UIMax = "5.0",
+		DisplayName = "Impulse Response Duration"))
+	float IndirectImpulseResponseDuration;
+
+	// How to spatialize indirect sound.
+	UPROPERTY(GlobalConfig, EditAnywhere, Category = Reverb)
+	EIplSpatializationMethod IndirectSpatializationMethod;
+
+	// The minimum distance between a source and a scene surface, used when calculating the energy received at the surface from the source during indirect sound simulation.
+	UPROPERTY(GlobalConfig, EditAnywhere, Category = Reverb, meta = (ClampMin = "0.1", ClampMax = "10.0", UIMin = "0.1", UIMax = "10.0"))
+	float IrradianceMinDistance;
+
+	// Maximum number of supported sources.
+	UPROPERTY(GlobalConfig, EditAnywhere, Category = Reverb, meta = (ClampMin = "1", ClampMax = "128", UIMin = "1", UIMax = "128"))
+	uint32 MaxSources;
+
+	//==============================================================================================================================================
+
+	// Preset quality settings for realtime simulation.
+	UPROPERTY(GlobalConfig, EditAnywhere, Category = "Reverb|Real-Time Quality Settings", meta = (DisplayName = "Quality Preset"))
+	EQualitySettings RealtimeQualityPreset;
+
+	// Percentage of CPU cores to use on an end user’s machine for performing real-time computation of environmental effects.
+	UPROPERTY(GlobalConfig, EditAnywhere, Category = "Reverb|Real-Time Quality Settings", meta = (DisplayName = "Real-time CPU Cores (%)", ClampMin = "1", ClampMax = "100", UIMin = "1", UIMax = "100"))
+	int32 RealTimeCPUCoresPercentage;
+
+	// Number of bounces for realtime simulation.
+	UPROPERTY(GlobalConfig, EditAnywhere, Category = "Reverb|Real-Time Quality Settings", meta = (ClampMin = "1", ClampMax = "32", 
+		UIMin = "1", UIMax = "32", DisplayName = "Bounces"))
+	int32 RealtimeBounces;
+
+	// Number of rays to trace for realtime simulation.
+	UPROPERTY(GlobalConfig, EditAnywhere, Category = "Reverb|Real-Time Quality Settings", meta = (ClampMin = "512", ClampMax = "16384", 
+		UIMin = "512", UIMax = "16384", DisplayName = "Rays"))
+	int32 RealtimeRays;
+
+	// Number of secondary rays to trace for realtime simulation.
+	UPROPERTY(GlobalConfig, EditAnywhere, Category = "Reverb|Real-Time Quality Settings", meta = (ClampMin = "128", ClampMax = "4096", 
+		UIMin = "128", UIMax = "4096", DisplayName = "Secondary Rays"))
+	int32 RealtimeSecondaryRays;
+
+	//==============================================================================================================================================
+
+	// Preset quality settings for baked simulation.
+	UPROPERTY(GlobalConfig, EditAnywhere, Category = "Reverb|Baked Quality Settings", meta = (DisplayName = "Quality Preset"))
+	EQualitySettings BakedQualityPreset;
+
+	// Percentage of CPU cores to use on a developer’s machine for baking environmental effects during the design phase.
+	UPROPERTY(GlobalConfig, EditAnywhere, Category = "Reverb|Baked Quality Settings", meta = (DisplayName = "Baking CPU Cores (%)", ClampMin = "1", ClampMax = "100", UIMin = "1", UIMax = "100"))
+	int32 BakingCPUCoresPercentage;
+
+	// Number of bounces for baked simulation.
+	UPROPERTY(GlobalConfig, EditAnywhere, Category = "Reverb|Baked Quality Settings", meta = (ClampMin = "16", ClampMax = "256", 
+		UIMin = "16", UIMax = "256", DisplayName = "Bounces"))
+	int32 BakedBounces;
+
+	// Number of rays to shoot for baked simulation.
+	UPROPERTY(GlobalConfig, EditAnywhere, Category = "Reverb|Baked Quality Settings", meta = (ClampMin = "8192", ClampMax = "65536", 
+		UIMin = "8192", UIMax = "65536", DisplayName = "Rays"))
+	int32 BakedRays;
+
+	// Number of secondary rays to shoot for baked simulation.
+	UPROPERTY(GlobalConfig, EditAnywhere, Category = "Reverb|Baked Quality Settings", meta = (ClampMin = "1024", ClampMax = "16384", 
+		UIMin = "1024", UIMax = "16384", DisplayName = "Secondary Rays"))
+	int32 BakedSecondaryRays;
+
+	//==============================================================================================================================================
+
+	// Maximum number of compute units to reserve on the GPU.
+	UPROPERTY(GlobalConfig, EditAnywhere, Category = "GPU Resource Reservation", meta = (ClampMin = "0", ClampMax = "16", UIMin = "0", UIMax = "16"))
+	int32 MaxComputeUnits;
+
+	// Fraction of maximum reserved CUs that should be used for impulse response (IR) update.
+	UPROPERTY(GlobalConfig, EditAnywhere, Category = "GPU Resource Reservation", meta = (ClampMin = "0", ClampMax = "1", UIMin = "0", UIMax = "1"))
+	float FractionComputeUnitsForIRUpdate;
+
+	//==============================================================================================================================================
+
+	// TAN override output of indirect propagation is stored in ambisonics of this order.
+	UPROPERTY(GlobalConfig, EditAnywhere, Category = "TrueAudio Next Overrides", meta = (ClampMin = "0", ClampMax = "3", UIMin = "0", UIMax = "3",
+		DisplayName = "Override Ambisonics Order"))
+	int32 TANIndirectImpulseResponseOrder;
+
+	// TAN override length of impulse response to compute for each sound source.
+	UPROPERTY(GlobalConfig, EditAnywhere, Category = "TrueAudio Next Overrides", meta = (ClampMin = "0.1", ClampMax = "5.0", UIMin = "0.1", UIMax = "5.0",
+		DisplayName = "Override Impulse Response Duration"))
+	float TANIndirectImpulseResponseDuration;
+
+	// TAN override maximum number of supported sources.
+	UPROPERTY(GlobalConfig, EditAnywhere, Category = "TrueAudio Next Overrides", meta = (ClampMin = "1", ClampMax = "128", UIMin = "1", UIMax = "128",
+		DisplayName = "Override Max Sources"))
+	uint32 TANMaxSources;
+
+	//==============================================================================================================================================
+
+	// This is the number of probes that are simultaneously baked on the GPU. Increasing this number results in better utilization of available GPU compute resources, 
+	// at the cost of increased GPU memory consumption. If this number is set too high, you may encounter errors when baking; if this happens, reduce the Baking Batch 
+	// Size value until baking succeeds.
+	UPROPERTY(GlobalConfig, EditAnywhere, Category = "Radeon Rays", meta = (ClampMin = "0", ClampMax = "16", UIMin = "0", UIMax = "16",
+		DisplayName = "Baking Batch Size"))
+	int32 RadeonRaysBakingBatchSize;
+};
\ No newline at end of file
diff --git a/Engine/Plugins/Runtime/Steam/SteamAudio/Source/SteamAudio/Private/PhononCommon.h b/Engine/Plugins/Runtime/Steam/SteamAudio/Source/SteamAudio/Private/PhononCommon.h
deleted file mode 100644
diff --git a/Engine/Plugins/Runtime/Steam/SteamAudio/Source/SteamAudio/Private/PhononMaterial.h b/Engine/Plugins/Runtime/Steam/SteamAudio/Source/SteamAudio/Private/PhononMaterial.h
deleted file mode 100644
diff --git a/Engine/Plugins/Runtime/Steam/SteamAudio/Source/SteamAudio/Private/PhononProbeComponent.h b/Engine/Plugins/Runtime/Steam/SteamAudio/Source/SteamAudio/Private/PhononProbeComponent.h
deleted file mode 100644
diff --git a/Engine/Plugins/Runtime/Steam/SteamAudio/Source/SteamAudio/Private/PhononProbeVolume.h b/Engine/Plugins/Runtime/Steam/SteamAudio/Source/SteamAudio/Private/PhononProbeVolume.h
deleted file mode 100644
diff --git a/Engine/Plugins/Runtime/Steam/SteamAudio/Source/SteamAudio/Private/PhononScene.h b/Engine/Plugins/Runtime/Steam/SteamAudio/Source/SteamAudio/Private/PhononScene.h
deleted file mode 100644
diff --git a/Engine/Plugins/Runtime/Steam/SteamAudio/Source/SteamAudio/Private/PhononSourceComponent.h b/Engine/Plugins/Runtime/Steam/SteamAudio/Source/SteamAudio/Private/PhononSourceComponent.h
deleted file mode 100644
diff --git a/Engine/Plugins/Runtime/Steam/SteamAudio/Source/SteamAudio/Private/SteamAudioEnvironment.h b/Engine/Plugins/Runtime/Steam/SteamAudio/Source/SteamAudio/Private/SteamAudioEnvironment.h
deleted file mode 100644
diff --git a/Engine/Plugins/Runtime/Steam/SteamAudio/Source/SteamAudio/Private/SteamAudioSettings.h b/Engine/Plugins/Runtime/Steam/SteamAudio/Source/SteamAudio/Private/SteamAudioSettings.h
deleted file mode 100644
diff --git a/Engine/Plugins/Runtime/Steam/SteamAudio/Source/SteamAudioEditor/Private/IndirectBaker.cpp b/Engine/Plugins/Runtime/Steam/SteamAudio/Source/SteamAudioEditor/Private/IndirectBaker.cpp
--- a/Engine/Plugins/Runtime/Steam/SteamAudio/Source/SteamAudioEditor/Private/IndirectBaker.cpp
+++ b/Engine/Plugins/Runtime/Steam/SteamAudio/Source/SteamAudioEditor/Private/IndirectBaker.cpp
@@ -9,7 +9,6 @@
 #include "PhononScene.h"
 #include "SteamAudioSettings.h"
 #include "SteamAudioEditorModule.h"
-#include "PhononReverb.h"
 #include "SteamAudioEnvironment.h"
 
 #include "LevelEditor.h"
diff --git a/Engine/Plugins/Runtime/Steam/SteamAudio/Source/SteamAudioEditor/Private/PhononProbeVolumeDetails.cpp b/Engine/Plugins/Runtime/Steam/SteamAudio/Source/SteamAudioEditor/Private/PhononProbeVolumeDetails.cpp
--- a/Engine/Plugins/Runtime/Steam/SteamAudio/Source/SteamAudioEditor/Private/PhononProbeVolumeDetails.cpp
+++ b/Engine/Plugins/Runtime/Steam/SteamAudio/Source/SteamAudioEditor/Private/PhononProbeVolumeDetails.cpp
@@ -7,7 +7,6 @@
 #include "PhononCommon.h"
 #include "SteamAudioEditorModule.h"
 #include "SteamAudioSettings.h"
-#include "PhononReverb.h"
 #include "SteamAudioEnvironment.h"
 
 #include "DetailLayoutBuilder.h"
diff --git a/Engine/Plugins/Runtime/Steam/SteamAudio/Source/SteamAudioEditor/SteamAudioEditor.Build.cs b/Engine/Plugins/Runtime/Steam/SteamAudio/Source/SteamAudioEditor/SteamAudioEditor.Build.cs
--- a/Engine/Plugins/Runtime/Steam/SteamAudio/Source/SteamAudioEditor/SteamAudioEditor.Build.cs
+++ b/Engine/Plugins/Runtime/Steam/SteamAudio/Source/SteamAudioEditor/SteamAudioEditor.Build.cs
@@ -8,12 +8,6 @@ namespace UnrealBuildTool.Rules
 	{
 		public SteamAudioEditor(ReadOnlyTargetRules Target) : base(Target)
 		{
-			PrivateIncludePaths.AddRange(
-				new string[] {
-					"SteamAudio/Private",
-				}
-			);
-
 			PublicDependencyModuleNames.AddRange(
 				new string[] {
 					"Core",
