Deprecate Steam Audio Plugin

diff --git a/Engine/Plugins/Runtime/Steam/SteamAudio/Source/SteamAudio/Private/SteamAudioModule.cpp b/Engine/Plugins/Runtime/Steam/SteamAudio/Source/SteamAudio/Private/SteamAudioModule.cpp
--- a/Engine/Plugins/Runtime/Steam/SteamAudio/Source/SteamAudio/Private/SteamAudioModule.cpp
+++ b/Engine/Plugins/Runtime/Steam/SteamAudio/Source/SteamAudio/Private/SteamAudioModule.cpp
@@ -31,6 +31,8 @@ namespace SteamAudio
 
 		UE_LOG(LogSteamAudio, Log, TEXT("FSteamAudioModule Startup"));
 
+		UE_LOG(LogSteamAudio, Warning, TEXT("The Steam Audio plugin is deprecated and will be removed in a future engine release. See the Plugins Browser for more information."))
+
 		// Make folder paths globally available
 		BasePath = FPaths::ProjectContentDir() + "SteamAudio/";
 		RuntimePath = BasePath + "Runtime/";
diff --git a/Engine/Plugins/Runtime/Steam/SteamAudio/SteamAudio.uplugin b/Engine/Plugins/Runtime/Steam/SteamAudio/SteamAudio.uplugin
--- a/Engine/Plugins/Runtime/Steam/SteamAudio/SteamAudio.uplugin
+++ b/Engine/Plugins/Runtime/Steam/SteamAudio/SteamAudio.uplugin
@@ -2,8 +2,8 @@
 	"FileVersion" : 3,
 	"Version" : 1,
 	"VersionName" : "2.0-beta.17",
-	"FriendlyName" : "Steam Audio",
-	"Description" : "Physically-based sound rendering.",
+	"FriendlyName" : "Steam Audio (Deprecated)",
+	"Description" : "Physically-based sound rendering. This plugin is deprecated and will be removed in a future engine release. Please use the plugin from Valve's website.",
 	"Category" : "Audio",
 	"CreatedBy" : "Valve Corporation",
 	"CreatedByURL" : "https://valvesoftware.github.io/steam-audio",
