Turn Steam Audio plugin into a stub plugin

diff --git a/Engine/Plugins/Runtime/Steam/SteamAudio/Source/SteamAudio/Internal/PhononCommon.h b/Engine/Plugins/Runtime/Steam/SteamAudio/Source/SteamAudio/Internal/PhononCommon.h
deleted file mode 100644
diff --git a/Engine/Plugins/Runtime/Steam/SteamAudio/Source/SteamAudio/Internal/PhononMaterial.h b/Engine/Plugins/Runtime/Steam/SteamAudio/Source/SteamAudio/Internal/PhononMaterial.h
deleted file mode 100644
diff --git a/Engine/Plugins/Runtime/Steam/SteamAudio/Source/SteamAudio/Internal/PhononProbeComponent.h b/Engine/Plugins/Runtime/Steam/SteamAudio/Source/SteamAudio/Internal/PhononProbeComponent.h
deleted file mode 100644
diff --git a/Engine/Plugins/Runtime/Steam/SteamAudio/Source/SteamAudio/Internal/PhononProbeVolume.h b/Engine/Plugins/Runtime/Steam/SteamAudio/Source/SteamAudio/Internal/PhononProbeVolume.h
deleted file mode 100644
diff --git a/Engine/Plugins/Runtime/Steam/SteamAudio/Source/SteamAudio/Internal/PhononScene.h b/Engine/Plugins/Runtime/Steam/SteamAudio/Source/SteamAudio/Internal/PhononScene.h
deleted file mode 100644
diff --git a/Engine/Plugins/Runtime/Steam/SteamAudio/Source/SteamAudio/Internal/PhononSourceComponent.h b/Engine/Plugins/Runtime/Steam/SteamAudio/Source/SteamAudio/Internal/PhononSourceComponent.h
deleted file mode 100644
diff --git a/Engine/Plugins/Runtime/Steam/SteamAudio/Source/SteamAudio/Internal/SteamAudioEnvironment.h b/Engine/Plugins/Runtime/Steam/SteamAudio/Source/SteamAudio/Internal/SteamAudioEnvironment.h
deleted file mode 100644
diff --git a/Engine/Plugins/Runtime/Steam/SteamAudio/Source/SteamAudio/Internal/SteamAudioSettings.h b/Engine/Plugins/Runtime/Steam/SteamAudio/Source/SteamAudio/Internal/SteamAudioSettings.h
deleted file mode 100644
diff --git a/Engine/Plugins/Runtime/Steam/SteamAudio/Source/SteamAudio/Private/PhononCommon.cpp b/Engine/Plugins/Runtime/Steam/SteamAudio/Source/SteamAudio/Private/PhononCommon.cpp
deleted file mode 100644
diff --git a/Engine/Plugins/Runtime/Steam/SteamAudio/Source/SteamAudio/Private/PhononGeometryComponent.cpp b/Engine/Plugins/Runtime/Steam/SteamAudio/Source/SteamAudio/Private/PhononGeometryComponent.cpp
deleted file mode 100644
diff --git a/Engine/Plugins/Runtime/Steam/SteamAudio/Source/SteamAudio/Private/PhononGeometryComponent.h b/Engine/Plugins/Runtime/Steam/SteamAudio/Source/SteamAudio/Private/PhononGeometryComponent.h
deleted file mode 100644
diff --git a/Engine/Plugins/Runtime/Steam/SteamAudio/Source/SteamAudio/Private/PhononMaterial.cpp b/Engine/Plugins/Runtime/Steam/SteamAudio/Source/SteamAudio/Private/PhononMaterial.cpp
deleted file mode 100644
diff --git a/Engine/Plugins/Runtime/Steam/SteamAudio/Source/SteamAudio/Private/PhononMaterialComponent.cpp b/Engine/Plugins/Runtime/Steam/SteamAudio/Source/SteamAudio/Private/PhononMaterialComponent.cpp
deleted file mode 100644
diff --git a/Engine/Plugins/Runtime/Steam/SteamAudio/Source/SteamAudio/Private/PhononMaterialComponent.h b/Engine/Plugins/Runtime/Steam/SteamAudio/Source/SteamAudio/Private/PhononMaterialComponent.h
deleted file mode 100644
diff --git a/Engine/Plugins/Runtime/Steam/SteamAudio/Source/SteamAudio/Private/PhononOcclusion.cpp b/Engine/Plugins/Runtime/Steam/SteamAudio/Source/SteamAudio/Private/PhononOcclusion.cpp
deleted file mode 100644
diff --git a/Engine/Plugins/Runtime/Steam/SteamAudio/Source/SteamAudio/Private/PhononOcclusion.h b/Engine/Plugins/Runtime/Steam/SteamAudio/Source/SteamAudio/Private/PhononOcclusion.h
deleted file mode 100644
diff --git a/Engine/Plugins/Runtime/Steam/SteamAudio/Source/SteamAudio/Private/PhononOcclusionSourceSettings.cpp b/Engine/Plugins/Runtime/Steam/SteamAudio/Source/SteamAudio/Private/PhononOcclusionSourceSettings.cpp
deleted file mode 100644
diff --git a/Engine/Plugins/Runtime/Steam/SteamAudio/Source/SteamAudio/Private/PhononPluginManager.cpp b/Engine/Plugins/Runtime/Steam/SteamAudio/Source/SteamAudio/Private/PhononPluginManager.cpp
deleted file mode 100644
diff --git a/Engine/Plugins/Runtime/Steam/SteamAudio/Source/SteamAudio/Private/PhononPluginManager.h b/Engine/Plugins/Runtime/Steam/SteamAudio/Source/SteamAudio/Private/PhononPluginManager.h
deleted file mode 100644
diff --git a/Engine/Plugins/Runtime/Steam/SteamAudio/Source/SteamAudio/Private/PhononProbeComponent.cpp b/Engine/Plugins/Runtime/Steam/SteamAudio/Source/SteamAudio/Private/PhononProbeComponent.cpp
deleted file mode 100644
diff --git a/Engine/Plugins/Runtime/Steam/SteamAudio/Source/SteamAudio/Private/PhononProbeVolume.cpp b/Engine/Plugins/Runtime/Steam/SteamAudio/Source/SteamAudio/Private/PhononProbeVolume.cpp
deleted file mode 100644
diff --git a/Engine/Plugins/Runtime/Steam/SteamAudio/Source/SteamAudio/Private/PhononReverb.cpp b/Engine/Plugins/Runtime/Steam/SteamAudio/Source/SteamAudio/Private/PhononReverb.cpp
deleted file mode 100644
diff --git a/Engine/Plugins/Runtime/Steam/SteamAudio/Source/SteamAudio/Private/PhononReverb.h b/Engine/Plugins/Runtime/Steam/SteamAudio/Source/SteamAudio/Private/PhononReverb.h
deleted file mode 100644
diff --git a/Engine/Plugins/Runtime/Steam/SteamAudio/Source/SteamAudio/Private/PhononReverbSourceSettings.cpp b/Engine/Plugins/Runtime/Steam/SteamAudio/Source/SteamAudio/Private/PhononReverbSourceSettings.cpp
deleted file mode 100644
diff --git a/Engine/Plugins/Runtime/Steam/SteamAudio/Source/SteamAudio/Private/PhononScene.cpp b/Engine/Plugins/Runtime/Steam/SteamAudio/Source/SteamAudio/Private/PhononScene.cpp
deleted file mode 100644
diff --git a/Engine/Plugins/Runtime/Steam/SteamAudio/Source/SteamAudio/Private/PhononSourceComponent.cpp b/Engine/Plugins/Runtime/Steam/SteamAudio/Source/SteamAudio/Private/PhononSourceComponent.cpp
deleted file mode 100644
diff --git a/Engine/Plugins/Runtime/Steam/SteamAudio/Source/SteamAudio/Private/PhononSpatialization.cpp b/Engine/Plugins/Runtime/Steam/SteamAudio/Source/SteamAudio/Private/PhononSpatialization.cpp
deleted file mode 100644
diff --git a/Engine/Plugins/Runtime/Steam/SteamAudio/Source/SteamAudio/Private/PhononSpatialization.h b/Engine/Plugins/Runtime/Steam/SteamAudio/Source/SteamAudio/Private/PhononSpatialization.h
deleted file mode 100644
diff --git a/Engine/Plugins/Runtime/Steam/SteamAudio/Source/SteamAudio/Private/PhononSpatializationSourceSettings.cpp b/Engine/Plugins/Runtime/Steam/SteamAudio/Source/SteamAudio/Private/PhononSpatializationSourceSettings.cpp
deleted file mode 100644
diff --git a/Engine/Plugins/Runtime/Steam/SteamAudio/Source/SteamAudio/Private/SteamAudioEnvironment.cpp b/Engine/Plugins/Runtime/Steam/SteamAudio/Source/SteamAudio/Private/SteamAudioEnvironment.cpp
deleted file mode 100644
diff --git a/Engine/Plugins/Runtime/Steam/SteamAudio/Source/SteamAudio/Private/SteamAudioModule.cpp b/Engine/Plugins/Runtime/Steam/SteamAudio/Source/SteamAudio/Private/SteamAudioModule.cpp
--- a/Engine/Plugins/Runtime/Steam/SteamAudio/Source/SteamAudio/Private/SteamAudioModule.cpp
+++ b/Engine/Plugins/Runtime/Steam/SteamAudio/Source/SteamAudio/Private/SteamAudioModule.cpp
@@ -3,16 +3,15 @@
 //
 
 #include "SteamAudioModule.h"
-#include "PhononCommon.h"
 #include "Misc/Paths.h"
 #include "Features/IModularFeatures.h"
 
+DEFINE_LOG_CATEGORY(LogSteamAudio);
+
 IMPLEMENT_MODULE(SteamAudio::FSteamAudioModule, SteamAudio)
 
 namespace SteamAudio
 {
-	void* FSteamAudioModule::PhononDllHandle = nullptr;
-
 	static bool bModuleStartedUp = false;
 
 	FSteamAudioModule::FSteamAudioModule()
@@ -31,34 +30,7 @@ namespace SteamAudio
 
 		UE_LOG(LogSteamAudio, Log, TEXT("FSteamAudioModule Startup"));
 
-		UE_LOG(LogSteamAudio, Warning, TEXT("The Steam Audio plugin is deprecated and will be removed in a future engine release. See the Plugins Browser for more information."))
-
-		// Make folder paths globally available
-		BasePath = FPaths::ProjectContentDir() + "SteamAudio/";
-		RuntimePath = BasePath + "Runtime/";
-		DynamicRuntimePath = RuntimePath + "Dynamic/";
-		EditorOnlyPath = BasePath + "EditorOnly/";
-		DynamicEditorOnlyPath = EditorOnlyPath + "Dynamic/";
-
-		//Register the Steam Audio plugin factories
-		IModularFeatures::Get().RegisterModularFeature(FSpatializationPluginFactory::GetModularFeatureName(), &SpatializationPluginFactory);
-		IModularFeatures::Get().RegisterModularFeature(FReverbPluginFactory::GetModularFeatureName(), &ReverbPluginFactory);
-		IModularFeatures::Get().RegisterModularFeature(FOcclusionPluginFactory::GetModularFeatureName(), &OcclusionPluginFactory);
-
-#if PLATFORM_WINDOWS
-		if (!FSteamAudioModule::PhononDllHandle)
-		{
-#if PLATFORM_32BITS
-			FString PathToDll = FPaths::EngineDir() / TEXT("Binaries/ThirdParty/Phonon/Win32/");
-#else
-			FString PathToDll = FPaths::EngineDir() / TEXT("Binaries/ThirdParty/Phonon/Win64/");
-#endif
-
-			FString DLLToLoad = PathToDll + TEXT("phonon.dll");
-			FSteamAudioModule::PhononDllHandle = LoadDll(DLLToLoad, true);
-		}
-#endif
-		iplCreateContext(UnrealLog, nullptr, nullptr, &GlobalContext);
+		UE_LOG(LogSteamAudio, Warning, TEXT("The Steam Audio plugin is deprecated and will be removed in a future engine release. Get Valve's version here: https://valvesoftware.github.io/steam-audio/doc/unreal/getting-started.html"))
 	}
 
 	void FSteamAudioModule::ShutdownModule()
@@ -68,80 +40,6 @@ namespace SteamAudio
 		check(bModuleStartedUp == true);
 
 		bModuleStartedUp = false;
-
-		iplDestroyContext(&GlobalContext);
-		iplCleanup();
-
-#if PLATFORM_WINDOWS
-		if (FSteamAudioModule::PhononDllHandle)
-		{
-			FPlatformProcess::FreeDllHandle(FSteamAudioModule::PhononDllHandle);
-			FSteamAudioModule::PhononDllHandle = nullptr;
-		}
-#endif
-	}
-
-	IAudioPluginFactory* FSteamAudioModule::GetPluginFactory(EAudioPlugin PluginType)
-	{
-		switch (PluginType)
-		{
-		case EAudioPlugin::SPATIALIZATION:
-			return &SpatializationPluginFactory;
-			break;
-		case EAudioPlugin::REVERB:
-			return &ReverbPluginFactory;
-			break;
-		case EAudioPlugin::OCCLUSION:
-			return &OcclusionPluginFactory;
-			break;
-		default:
-			return nullptr;
-			break;
-		}
-	}
-
-	void FSteamAudioModule::RegisterAudioDevice(FAudioDevice* AudioDeviceHandle)
-	{
-		if (!RegisteredAudioDevices.Contains(AudioDeviceHandle))
-		{
-			TAudioPluginListenerPtr NewPhononPluginManager = TAudioPluginListenerPtr(new FPhononPluginManager());
-			AudioDeviceHandle->RegisterPluginListener(NewPhononPluginManager);
-			RegisteredAudioDevices.Add(AudioDeviceHandle);
-		}
-	}
-
-	void FSteamAudioModule::UnregisterAudioDevice(FAudioDevice* AudioDeviceHandle)
-	{
-		RegisteredAudioDevices.Remove(AudioDeviceHandle);
-	}
-
-	TAudioOcclusionPtr FOcclusionPluginFactory::CreateNewOcclusionPlugin(FAudioDevice* OwningDevice)
-	{
-		//Register audio device to the Steam Module
-		FSteamAudioModule* Module = &FModuleManager::GetModuleChecked<FSteamAudioModule>("SteamAudio");
-		if (Module != nullptr)
-		{
-			Module->RegisterAudioDevice(OwningDevice);
-		}
-		
-		return TAudioOcclusionPtr(new FPhononOcclusion());
-	}
-
-	TAudioReverbPtr FReverbPluginFactory::CreateNewReverbPlugin(FAudioDevice* OwningDevice)
-	{
-		//Register the audio device to the steam module:
-		FSteamAudioModule* Module = &FModuleManager::GetModuleChecked<FSteamAudioModule>("SteamAudio");
-		if (Module != nullptr)
-		{
-			Module->RegisterAudioDevice(OwningDevice);
-		}
-
-		return TAudioReverbPtr(new FPhononReverb());
-	}
-
-	TAudioSpatializationPtr FSpatializationPluginFactory::CreateNewSpatializationPlugin(FAudioDevice* OwningDevice)
-	{
-		return TAudioSpatializationPtr(new FPhononSpatialization());
 	}
 
 }
diff --git a/Engine/Plugins/Runtime/Steam/SteamAudio/Source/SteamAudio/Private/SteamAudioModule.h b/Engine/Plugins/Runtime/Steam/SteamAudio/Source/SteamAudio/Private/SteamAudioModule.h
--- a/Engine/Plugins/Runtime/Steam/SteamAudio/Source/SteamAudio/Private/SteamAudioModule.h
+++ b/Engine/Plugins/Runtime/Steam/SteamAudio/Source/SteamAudio/Private/SteamAudioModule.h
@@ -5,87 +5,13 @@
 #pragma once
 
 #include "ISteamAudioModule.h"
-#include "PhononSpatialization.h"
-#include "PhononOcclusion.h"
-#include "PhononReverb.h"
-#include "PhononPluginManager.h"
-#include "PhononReverbSourceSettings.h"
-#include "PhononOcclusionSourceSettings.h"
-#include "PhononSpatializationSourceSettings.h"
 #include "AudioPluginUtilities.h"
 #include "AudioDevice.h"
 
+DECLARE_LOG_CATEGORY_EXTERN(LogSteamAudio, Log, All);
+
 namespace SteamAudio
 {
-	/************************************************************************/
-	/* Plugin Factories registered with Steam Audio Module                  */
-	/************************************************************************/
-	class FSpatializationPluginFactory : public IAudioSpatializationFactory
-	{
-	public:
-		virtual FString GetDisplayName() override
-		{
-			static FString DisplayName = FString(TEXT("Steam Audio"));
-			return DisplayName;
-		}
-
-		virtual bool SupportsPlatform(const FString& PlatformName) override
-		{
-			return	PlatformName == FString(TEXT("Windows")) || PlatformName == FString(TEXT("Android"));
-		}
-
-		virtual TAudioSpatializationPtr CreateNewSpatializationPlugin(FAudioDevice* OwningDevice) override;
-
-		virtual UClass* GetCustomSpatializationSettingsClass() const override
-		{ 
-			return UPhononOcclusionSourceSettings::StaticClass(); 
-		}
-	};
-
-	class FReverbPluginFactory : public IAudioReverbFactory
-	{
-	public:
-		virtual FString GetDisplayName() override
-		{
-			static FString DisplayName = FString(TEXT("Steam Audio"));
-			return DisplayName;
-		}
-
-		virtual bool SupportsPlatform(const FString& PlatformName) override
-		{
-			return	PlatformName == FString(TEXT("Windows"));
-		}
-
-		virtual TAudioReverbPtr CreateNewReverbPlugin(FAudioDevice* OwningDevice) override;
-
-		virtual UClass* GetCustomReverbSettingsClass() const
-		{
-			return UPhononReverbSourceSettings::StaticClass();
-		}
-	};
-
-	class FOcclusionPluginFactory : public IAudioOcclusionFactory
-	{
-	public:
-		virtual FString GetDisplayName() override
-		{
-			static FString DisplayName = FString(TEXT("Steam Audio"));
-			return DisplayName;
-		}
-
-		virtual bool SupportsPlatform(const FString& PlatformName) override
-		{
-			return	PlatformName == FString(TEXT("Windows")) || PlatformName == FString(TEXT("Android"));
-		}
-
-		virtual TAudioOcclusionPtr CreateNewOcclusionPlugin(FAudioDevice* OwningDevice) override;
-
-		virtual UClass* GetCustomOcclusionSettingsClass() const
-		{
-			return UPhononOcclusionSourceSettings::StaticClass();
-		}
-	};
-
 	/************************************************************************/
 	/* FSteamAudioModule                                                    */
 	/* Module interface for Steam Audio. Registers the Plugin Factories     */
@@ -100,40 +26,5 @@ namespace SteamAudio
 		virtual void StartupModule() override;
 		virtual void ShutdownModule() override;
 		//~ End IAudioPlugin
-
-		/** 
-		*  Returns pointer to each type of plugin factory with Steam Audio.
-		*  returns nullptr if given an invalid parameter for PluginType.
-		*/
-		IAudioPluginFactory* GetPluginFactory(EAudioPlugin PluginType);
-
-		/*
-		 * Registers a given audio device to the Steam Audio Device.
-		 * Every audio device running Steam Audio's reverb or occlusion manager
-		 * requires a PhononPluginManager registered to it.
-		 */
-		void RegisterAudioDevice(FAudioDevice* AudioDeviceHandle);
-
-		/*
-		* Unregisters an audio device from the module.
-		*/
-		void UnregisterAudioDevice(FAudioDevice* AudioDeviceHandle);
-
-	private:
-		static void* PhononDllHandle;
-		static void* TANDllHandle;
-		static void* TANUtilsDllHandle;
-		static void* EmbreeDllHandle;
-		static void* RadeonRaysDllHandle;
-		static void* TbbDllHandle;
-		static void* TbbMallocDllHandle;
-
-
-		TArray<FAudioDevice*> RegisteredAudioDevices;
-
-		//Factories:
-		FSpatializationPluginFactory SpatializationPluginFactory;
-		FReverbPluginFactory ReverbPluginFactory;
-		FOcclusionPluginFactory OcclusionPluginFactory;
 	};
 }
\ No newline at end of file
diff --git a/Engine/Plugins/Runtime/Steam/SteamAudio/Source/SteamAudio/Private/SteamAudioSettings.cpp b/Engine/Plugins/Runtime/Steam/SteamAudio/Source/SteamAudio/Private/SteamAudioSettings.cpp
deleted file mode 100644
diff --git a/Engine/Plugins/Runtime/Steam/SteamAudio/Source/SteamAudio/Public/PhononOcclusionSourceSettings.h b/Engine/Plugins/Runtime/Steam/SteamAudio/Source/SteamAudio/Public/PhononOcclusionSourceSettings.h
deleted file mode 100644
diff --git a/Engine/Plugins/Runtime/Steam/SteamAudio/Source/SteamAudio/Public/PhononReverbSourceSettings.h b/Engine/Plugins/Runtime/Steam/SteamAudio/Source/SteamAudio/Public/PhononReverbSourceSettings.h
deleted file mode 100644
diff --git a/Engine/Plugins/Runtime/Steam/SteamAudio/Source/SteamAudio/Public/PhononSpatializationSourceSettings.h b/Engine/Plugins/Runtime/Steam/SteamAudio/Source/SteamAudio/Public/PhononSpatializationSourceSettings.h
deleted file mode 100644
diff --git a/Engine/Plugins/Runtime/Steam/SteamAudio/Source/SteamAudioEditor/Private/BakeIndirectWindow.cpp b/Engine/Plugins/Runtime/Steam/SteamAudio/Source/SteamAudioEditor/Private/BakeIndirectWindow.cpp
deleted file mode 100644
diff --git a/Engine/Plugins/Runtime/Steam/SteamAudio/Source/SteamAudioEditor/Private/BakeIndirectWindow.h b/Engine/Plugins/Runtime/Steam/SteamAudio/Source/SteamAudioEditor/Private/BakeIndirectWindow.h
deleted file mode 100644
diff --git a/Engine/Plugins/Runtime/Steam/SteamAudio/Source/SteamAudioEditor/Private/IndirectBaker.cpp b/Engine/Plugins/Runtime/Steam/SteamAudio/Source/SteamAudioEditor/Private/IndirectBaker.cpp
deleted file mode 100644
diff --git a/Engine/Plugins/Runtime/Steam/SteamAudio/Source/SteamAudioEditor/Private/IndirectBaker.h b/Engine/Plugins/Runtime/Steam/SteamAudio/Source/SteamAudioEditor/Private/IndirectBaker.h
deleted file mode 100644
diff --git a/Engine/Plugins/Runtime/Steam/SteamAudio/Source/SteamAudioEditor/Private/PhononOcclusionSettingsFactory.cpp b/Engine/Plugins/Runtime/Steam/SteamAudio/Source/SteamAudioEditor/Private/PhononOcclusionSettingsFactory.cpp
deleted file mode 100644
diff --git a/Engine/Plugins/Runtime/Steam/SteamAudio/Source/SteamAudioEditor/Private/PhononOcclusionSettingsFactory.h b/Engine/Plugins/Runtime/Steam/SteamAudio/Source/SteamAudioEditor/Private/PhononOcclusionSettingsFactory.h
deleted file mode 100644
diff --git a/Engine/Plugins/Runtime/Steam/SteamAudio/Source/SteamAudioEditor/Private/PhononProbeComponentVisualizer.cpp b/Engine/Plugins/Runtime/Steam/SteamAudio/Source/SteamAudioEditor/Private/PhononProbeComponentVisualizer.cpp
deleted file mode 100644
diff --git a/Engine/Plugins/Runtime/Steam/SteamAudio/Source/SteamAudioEditor/Private/PhononProbeComponentVisualizer.h b/Engine/Plugins/Runtime/Steam/SteamAudio/Source/SteamAudioEditor/Private/PhononProbeComponentVisualizer.h
deleted file mode 100644
diff --git a/Engine/Plugins/Runtime/Steam/SteamAudio/Source/SteamAudioEditor/Private/PhononProbeVolumeDetails.cpp b/Engine/Plugins/Runtime/Steam/SteamAudio/Source/SteamAudioEditor/Private/PhononProbeVolumeDetails.cpp
deleted file mode 100644
diff --git a/Engine/Plugins/Runtime/Steam/SteamAudio/Source/SteamAudioEditor/Private/PhononProbeVolumeDetails.h b/Engine/Plugins/Runtime/Steam/SteamAudio/Source/SteamAudioEditor/Private/PhononProbeVolumeDetails.h
deleted file mode 100644
diff --git a/Engine/Plugins/Runtime/Steam/SteamAudio/Source/SteamAudioEditor/Private/PhononReverbSettingsFactory.cpp b/Engine/Plugins/Runtime/Steam/SteamAudio/Source/SteamAudioEditor/Private/PhononReverbSettingsFactory.cpp
deleted file mode 100644
diff --git a/Engine/Plugins/Runtime/Steam/SteamAudio/Source/SteamAudioEditor/Private/PhononReverbSettingsFactory.h b/Engine/Plugins/Runtime/Steam/SteamAudio/Source/SteamAudioEditor/Private/PhononReverbSettingsFactory.h
deleted file mode 100644
diff --git a/Engine/Plugins/Runtime/Steam/SteamAudio/Source/SteamAudioEditor/Private/PhononSourceComponentDetails.cpp b/Engine/Plugins/Runtime/Steam/SteamAudio/Source/SteamAudioEditor/Private/PhononSourceComponentDetails.cpp
deleted file mode 100644
diff --git a/Engine/Plugins/Runtime/Steam/SteamAudio/Source/SteamAudioEditor/Private/PhononSourceComponentDetails.h b/Engine/Plugins/Runtime/Steam/SteamAudio/Source/SteamAudioEditor/Private/PhononSourceComponentDetails.h
deleted file mode 100644
diff --git a/Engine/Plugins/Runtime/Steam/SteamAudio/Source/SteamAudioEditor/Private/PhononSourceComponentVisualizer.cpp b/Engine/Plugins/Runtime/Steam/SteamAudio/Source/SteamAudioEditor/Private/PhononSourceComponentVisualizer.cpp
deleted file mode 100644
diff --git a/Engine/Plugins/Runtime/Steam/SteamAudio/Source/SteamAudioEditor/Private/PhononSourceComponentVisualizer.h b/Engine/Plugins/Runtime/Steam/SteamAudio/Source/SteamAudioEditor/Private/PhononSourceComponentVisualizer.h
deleted file mode 100644
diff --git a/Engine/Plugins/Runtime/Steam/SteamAudio/Source/SteamAudioEditor/Private/PhononSpatializationSettingsFactory.cpp b/Engine/Plugins/Runtime/Steam/SteamAudio/Source/SteamAudioEditor/Private/PhononSpatializationSettingsFactory.cpp
deleted file mode 100644
diff --git a/Engine/Plugins/Runtime/Steam/SteamAudio/Source/SteamAudioEditor/Private/PhononSpatializationSettingsFactory.h b/Engine/Plugins/Runtime/Steam/SteamAudio/Source/SteamAudioEditor/Private/PhononSpatializationSettingsFactory.h
deleted file mode 100644
diff --git a/Engine/Plugins/Runtime/Steam/SteamAudio/Source/SteamAudioEditor/Private/SteamAudioEdMode.cpp b/Engine/Plugins/Runtime/Steam/SteamAudio/Source/SteamAudioEditor/Private/SteamAudioEdMode.cpp
deleted file mode 100644
diff --git a/Engine/Plugins/Runtime/Steam/SteamAudio/Source/SteamAudioEditor/Private/SteamAudioEdMode.h b/Engine/Plugins/Runtime/Steam/SteamAudio/Source/SteamAudioEditor/Private/SteamAudioEdMode.h
deleted file mode 100644
diff --git a/Engine/Plugins/Runtime/Steam/SteamAudio/Source/SteamAudioEditor/Private/SteamAudioEdModeToolkit.cpp b/Engine/Plugins/Runtime/Steam/SteamAudio/Source/SteamAudioEditor/Private/SteamAudioEdModeToolkit.cpp
deleted file mode 100644
diff --git a/Engine/Plugins/Runtime/Steam/SteamAudio/Source/SteamAudioEditor/Private/SteamAudioEdModeToolkit.h b/Engine/Plugins/Runtime/Steam/SteamAudio/Source/SteamAudioEditor/Private/SteamAudioEdModeToolkit.h
deleted file mode 100644
diff --git a/Engine/Plugins/Runtime/Steam/SteamAudio/Source/SteamAudioEditor/Private/SteamAudioEditorModule.cpp b/Engine/Plugins/Runtime/Steam/SteamAudio/Source/SteamAudioEditor/Private/SteamAudioEditorModule.cpp
deleted file mode 100644
diff --git a/Engine/Plugins/Runtime/Steam/SteamAudio/Source/SteamAudioEditor/Private/SteamAudioEditorModule.h b/Engine/Plugins/Runtime/Steam/SteamAudio/Source/SteamAudioEditor/Private/SteamAudioEditorModule.h
deleted file mode 100644
diff --git a/Engine/Plugins/Runtime/Steam/SteamAudio/Source/SteamAudioEditor/Private/TickableNotification.cpp b/Engine/Plugins/Runtime/Steam/SteamAudio/Source/SteamAudioEditor/Private/TickableNotification.cpp
deleted file mode 100644
diff --git a/Engine/Plugins/Runtime/Steam/SteamAudio/Source/SteamAudioEditor/Private/TickableNotification.h b/Engine/Plugins/Runtime/Steam/SteamAudio/Source/SteamAudioEditor/Private/TickableNotification.h
deleted file mode 100644
diff --git a/Engine/Plugins/Runtime/Steam/SteamAudio/Source/SteamAudioEditor/Public/ISteamAudioEditorModule.h b/Engine/Plugins/Runtime/Steam/SteamAudio/Source/SteamAudioEditor/Public/ISteamAudioEditorModule.h
deleted file mode 100644
diff --git a/Engine/Plugins/Runtime/Steam/SteamAudio/Source/SteamAudioEditor/SteamAudioEditor.Build.cs b/Engine/Plugins/Runtime/Steam/SteamAudio/Source/SteamAudioEditor/SteamAudioEditor.Build.cs
deleted file mode 100644
diff --git a/Engine/Plugins/Runtime/Steam/SteamAudio/SteamAudio.uplugin b/Engine/Plugins/Runtime/Steam/SteamAudio/SteamAudio.uplugin
--- a/Engine/Plugins/Runtime/Steam/SteamAudio/SteamAudio.uplugin
+++ b/Engine/Plugins/Runtime/Steam/SteamAudio/SteamAudio.uplugin
@@ -3,7 +3,7 @@
 	"Version" : 1,
 	"VersionName" : "2.0-beta.17",
 	"FriendlyName" : "Steam Audio (Deprecated)",
-	"Description" : "Physically-based sound rendering. This plugin is deprecated and will be removed in a future engine release. Please use the plugin from Valve's website.",
+	"Description" : "This plugin is deprecated and will be removed in a future engine release. Please use the plugin from Valve's website.",
 	"Category" : "Audio",
 	"CreatedBy" : "Valve Corporation",
 	"CreatedByURL" : "https://valvesoftware.github.io/steam-audio",
@@ -24,15 +24,6 @@
 			[
 				"Win64", "Android"
 			]
-		},
-		{
-			"Name" : "SteamAudioEditor",
-			"Type" : "Editor",
-			"LoadingPhase" : "PostEngineInit",
-			"PlatformAllowList" :
-			[
-				"Win64"
-			]
 		}
 	]
 }
