[Backout] - CL25430596
Original CL Desc
-----------------------------------------------------------------
Turn Steam Audio plugin into a stub plugin

diff --git a/Engine/Plugins/Runtime/Steam/SteamAudio/Source/SteamAudio/Private/PhononCommon.cpp b/Engine/Plugins/Runtime/Steam/SteamAudio/Source/SteamAudio/Private/PhononCommon.cpp
new file mode 100644
--- /dev/null
+++ b/Engine/Plugins/Runtime/Steam/SteamAudio/Source/SteamAudio/Private/PhononCommon.cpp
@@ -0,0 +1,374 @@
+//
+// Copyright (C) Valve Corporation. All rights reserved.
+//
+
+#include "PhononCommon.h"
+#include "Misc/Paths.h"
+#include "EngineUtils.h"
+#include "HAL/FileManager.h"
+#include "Interfaces/IPluginManager.h"
+#include "HAL/PlatformProcess.h"
+#include "Internationalization/Regex.h"
+
+DEFINE_LOG_CATEGORY(LogSteamAudio);
+
+static TMap<EQualitySettings, SteamAudio::FSimulationQualitySettings> GetRealtimeQualityPresets()
+{
+	TMap<EQualitySettings, SteamAudio::FSimulationQualitySettings> Presets;
+	SteamAudio::FSimulationQualitySettings Settings;
+
+	Settings.Bounces = 2;
+	Settings.Rays = 4096;
+	Settings.SecondaryRays = 512;
+	Presets.Add(EQualitySettings::LOW, Settings);
+
+	Settings.Bounces = 4;
+	Settings.Rays = 8192;
+	Settings.SecondaryRays = 1024;
+	Presets.Add(EQualitySettings::MEDIUM, Settings);
+
+	Settings.Bounces = 8;
+	Settings.Rays = 16384;
+	Settings.SecondaryRays = 2048;
+	Presets.Add(EQualitySettings::HIGH, Settings);
+
+	Settings.Bounces = 0;
+	Settings.Rays = 0;
+	Settings.SecondaryRays = 0;
+	Presets.Add(EQualitySettings::CUSTOM, Settings);
+
+	return Presets;
+}
+
+static TMap<EQualitySettings, SteamAudio::FSimulationQualitySettings> GetBakedQualityPresets()
+{
+	TMap<EQualitySettings, SteamAudio::FSimulationQualitySettings> Presets;
+	SteamAudio::FSimulationQualitySettings Settings;
+
+	Settings.Bounces = 64;
+	Settings.Rays = 16384;
+	Settings.SecondaryRays = 2048;
+	Presets.Add(EQualitySettings::LOW, Settings);
+
+	Settings.Bounces = 128;
+	Settings.Rays = 32768;
+	Settings.SecondaryRays = 4096;
+	Presets.Add(EQualitySettings::MEDIUM, Settings);
+
+	Settings.Bounces = 256;
+	Settings.Rays = 65536;
+	Settings.SecondaryRays = 8192;
+	Presets.Add(EQualitySettings::HIGH, Settings);
+
+	Settings.Bounces = 0;
+	Settings.Rays = 0;
+	Settings.SecondaryRays = 0;
+	Presets.Add(EQualitySettings::CUSTOM, Settings);
+
+	return Presets;
+}
+
+namespace SteamAudio
+{
+	TMap<EQualitySettings, FSimulationQualitySettings> RealtimeSimulationQualityPresets = GetRealtimeQualityPresets();
+	TMap<EQualitySettings, FSimulationQualitySettings> BakedSimulationQualityPresets = GetBakedQualityPresets();
+
+	FString BasePath;
+	FString RuntimePath;
+	FString DynamicRuntimePath;
+	FString EditorOnlyPath;
+	FString DynamicEditorOnlyPath;
+
+	void* UnrealAlloc(const size_t size, const size_t alignment)
+	{
+		return FMemory::Malloc(size, alignment);
+	}
+
+	void UnrealFree(void* ptr)
+	{
+		FMemory::Free(ptr);
+	}
+
+	void UnrealLog(char* msg)
+	{
+		FString Message(msg);
+		UE_LOG(LogSteamAudio, Log, TEXT("%s"), *Message); 
+	}
+
+	IPLhandle GlobalContext = nullptr;
+
+	FVector UnrealToPhononFVector(const FVector& UnrealCoords, const bool bScale)
+	{
+		FVector PhononCoords;
+		PhononCoords.X = UnrealCoords.Y;
+		PhononCoords.Y = UnrealCoords.Z; 
+		PhononCoords.Z = -UnrealCoords.X; 
+		return bScale ? PhononCoords * SCALEFACTOR : PhononCoords;
+	}
+
+	IPLVector3 UnrealToPhononIPLVector3(const FVector& UnrealCoords, const bool bScale)
+	{
+		IPLVector3 PhononCoords;
+		PhononCoords.x = UnrealCoords.Y;
+		PhononCoords.y = UnrealCoords.Z;
+		PhononCoords.z = -UnrealCoords.X;
+		
+		if (bScale)
+		{
+			PhononCoords.x *= SCALEFACTOR;
+			PhononCoords.y *= SCALEFACTOR;
+			PhononCoords.z *= SCALEFACTOR;
+		}
+
+		return PhononCoords;
+	}
+
+	FVector PhononToUnrealFVector(const FVector& coords, const bool bScale)
+	{
+		FVector UnrealCoords;
+		UnrealCoords.X = -coords.Z;
+		UnrealCoords.Y = coords.X;
+		UnrealCoords.Z = coords.Y;
+		
+		if (bScale)
+		{
+			UnrealCoords.X /= SCALEFACTOR;
+			UnrealCoords.Y /= SCALEFACTOR;
+			UnrealCoords.Z /= SCALEFACTOR;
+		}
+
+		return UnrealCoords;
+	}
+
+	IPLVector3 PhononToUnrealIPLVector3(const FVector& coords, const bool bScale)
+	{
+		IPLVector3 UnrealCoords = {0, 0, 0};
+
+		// TODO
+
+		return UnrealCoords;
+	}
+
+	IPLVector3 IPLVector3FromFVector(const FVector& Coords)
+	{
+		IPLVector3 IplVector;
+		IplVector.x = Coords.X;
+		IplVector.y = Coords.Y;
+		IplVector.z = Coords.Z;
+		return IplVector;
+	}
+
+	FVector FVectorFromIPLVector3(const IPLVector3& Coords)
+	{
+		FVector Vector;
+		Vector.X = Coords.x;
+		Vector.Y = Coords.y;
+		Vector.Z = Coords.z;
+		return Vector;
+	}
+
+	/**
+	 * Given a Unreal transform, convert it to a 4x4 column-major transformation matrix. OutMatrix is assumed to be a contiguous array of
+	 * 16 floats.
+	 */
+	void GetMatrixForTransform(const FTransform& Transform, float* OutMatrix, bool bOutputRowMajor /* = false */, bool bApplyScale /* = true */)
+	{
+		check(OutMatrix);
+
+		auto PhononTranslation = SteamAudio::UnrealToPhononFVector(Transform.GetTranslation());
+		auto PhononScale = SteamAudio::UnrealToPhononFVector(Transform.GetScale3D(), bApplyScale);
+
+		FQuat RotationQuatConverted;
+		RotationQuatConverted.X = -Transform.GetRotation().Y;
+		RotationQuatConverted.Y = -Transform.GetRotation().Z;
+		RotationQuatConverted.Z = Transform.GetRotation().X;
+		RotationQuatConverted.W = Transform.GetRotation().W;
+
+		FQuatRotationTranslationMatrix RotationTranslationMatrix(RotationQuatConverted, PhononTranslation);
+		FScaleMatrix ScaleMatrix(PhononScale);
+		auto ConvertedMatrix = (ScaleMatrix * RotationTranslationMatrix).GetTransposed();
+		
+		// Convert to column major if bOutputRowMajor is false, otherwise output column major
+		if (bOutputRowMajor)
+		{
+			// Row major order (flatten)
+			for (int32 i = 0; i < 4; ++i)
+			{
+				for (int32 j = 0; j < 4; ++j)
+				{
+					OutMatrix[i * 4 + j] = ConvertedMatrix.M[i][j];
+				}
+			}		
+		}
+		else
+		{
+			// Column major order (conversion from row major order)
+			for (int32 i = 0; i < 4; ++i)
+			{
+				for (int32 j = 0; j < 4; ++j)
+				{
+					OutMatrix[j * 4 + i] = ConvertedMatrix.M[i][j];
+				}
+			}
+		}
+	}
+
+	IPLMatrix4x4 GetIPLMatrix(float* InMatrix)
+	{
+		check(InMatrix);
+
+		IPLMatrix4x4 PhononMatrix;
+		FMemory::Memcpy(&PhononMatrix, InMatrix, sizeof(PhononMatrix));
+		
+		return PhononMatrix;
+	}
+
+	FString StrippedMapName(const FString& MapName)
+	{
+		const FRegexPattern PieMapPattern(TEXT("UEDPIE_\\d_"));
+		FRegexMatcher Matcher(PieMapPattern, MapName);
+		FString StrippedMapName = MapName;
+
+		if (Matcher.FindNext())
+		{
+			int32 Beginning = Matcher.GetMatchBeginning();
+			int32 Ending = Matcher.GetMatchEnding();
+			StrippedMapName = MapName.Mid(Ending, MapName.Len());
+		}
+
+		return StrippedMapName;
+	}
+
+	void* LoadDll(const FString& DllFile, bool bErrorOnLoadFailure)
+	{
+		UE_LOG(LogSteamAudio, Display, TEXT("Attempting to load %s"), *DllFile);
+
+		void* DllHandle = nullptr;
+
+		if (FPaths::FileExists(DllFile))
+		{
+			DllHandle = FPlatformProcess::GetDllHandle(*DllFile);
+		}
+		else
+		{
+			if (bErrorOnLoadFailure)
+			{
+				UE_LOG(LogSteamAudio, Error, TEXT("Failed to load missing file. %s"), *DllFile);
+			}
+			else
+			{
+				UE_LOG(LogSteamAudio, Display, TEXT("File does not exist. %s"), *DllFile);
+			}
+		}
+
+		if (!DllHandle)
+		{
+			if (bErrorOnLoadFailure)
+			{
+				UE_LOG(LogSteamAudio, Error, TEXT("Failure to load %s."), *DllFile)
+			}
+			else
+			{
+				UE_LOG(LogSteamAudio, Display, TEXT("Unable to load %s."), *DllFile);
+			}
+		}
+		else
+		{
+			UE_LOG(LogSteamAudio, Display, TEXT("Loaded %s."), *DllFile);
+		}
+
+		return DllHandle;
+	}
+
+	void LogSteamAudioStatus(const IPLerror Status)
+	{
+		if (Status != IPL_STATUS_SUCCESS)
+		{
+			UE_LOG(LogSteamAudio, Error, TEXT("Error: %s"), *StatusToFString(Status));
+		}
+	}
+
+	/**
+	 * Returns a string representing the given status.
+	 */
+	FString StatusToFString(const IPLerror Status)
+	{
+		switch (Status)
+		{
+		case IPL_STATUS_SUCCESS:
+			return FString(TEXT("Success."));
+		case IPL_STATUS_FAILURE:
+			return FString(TEXT("Failure."));
+		case IPL_STATUS_OUTOFMEMORY:
+			return FString(TEXT("Out of memory."));
+		case IPL_STATUS_INITIALIZATION:
+			return FString(TEXT("Initialization error."));
+		default:
+			return FString(TEXT("Unknown error."));
+		}
+	}
+
+	void ClosestHit(const IPLfloat32* Origin, const IPLfloat32* Direction, const IPLfloat32 MinDistance, const IPLfloat32 MaxDistance,
+		IPLfloat32* HitDistance, IPLfloat32* HitNormal, IPLMaterial** HitMaterial, IPLvoid* UserData)
+	{
+		UWorld* World = (UWorld*)UserData;
+
+		// PhysX doesn't handle infinity properly, leads to NaNs
+		float AdjustedMax = MaxDistance;
+
+		if (MaxDistance > 999999.0f)
+			AdjustedMax = 999999.0f;
+		else if (MaxDistance < -999999.0f)
+			AdjustedMax = -999999.0f;
+
+		FVector PhononOrigin(Origin[0], Origin[1], Origin[2]);
+		FVector PhononDirection(Direction[0], Direction[1], Direction[2]);
+		FVector PhononStart = PhononOrigin + PhononDirection * MinDistance;
+		FVector PhononEnd = PhononOrigin + PhononDirection * AdjustedMax;
+
+		FVector Start = PhononToUnrealFVector(PhononStart);
+		FVector End = PhononToUnrealFVector(PhononEnd);
+		FCollisionQueryParams TraceParams;
+		TraceParams.bTraceComplex = true;
+
+		FHitResult Result;
+		bool HitStatic = World->LineTraceSingleByObjectType(Result, Start, End, FCollisionObjectQueryParams(ECC_WorldStatic), TraceParams);
+
+		*HitDistance = Result.Distance * SCALEFACTOR;
+
+		FVector PhononNormal = UnrealToPhononFVector(Result.Normal, false);
+		HitNormal[0] = PhononNormal.X;
+		HitNormal[1] = PhononNormal.Y;
+		HitNormal[2] = PhononNormal.Z;
+		*HitMaterial = nullptr;
+	}
+
+	void AnyHit(const IPLfloat32* Origin, const IPLfloat32* Direction, const IPLfloat32 MinDistance, const IPLfloat32 MaxDistance,
+		IPLint32* HitExists, IPLvoid* UserData)
+	{
+		UWorld* World = (UWorld*)UserData;
+
+		// PhysX doesn't handle infinity properly, leads to NaNs
+		float AdjustedMax = MaxDistance;
+
+		if (MaxDistance > 999999.0f)
+			AdjustedMax = 999999.0f;
+		else if (MaxDistance < -999999.0f)
+			AdjustedMax = -999999.0f;
+
+		FVector PhononOrigin(Origin[0], Origin[1], Origin[2]);
+		FVector PhononDirection(Direction[0], Direction[1], Direction[2]);
+		FVector PhononStart = PhononOrigin + PhononDirection * MinDistance;
+		FVector PhononEnd = PhononOrigin + PhononDirection * AdjustedMax;
+
+		FVector Start = PhononToUnrealFVector(PhononStart);
+		FVector End = PhononToUnrealFVector(PhononEnd);
+		FCollisionQueryParams TraceParams;
+		TraceParams.bTraceComplex = true;
+
+		FHitResult Result;
+		bool HitStatic = World->LineTraceSingleByObjectType(Result, Start, End, FCollisionObjectQueryParams(ECC_WorldStatic), TraceParams);
+
+		*HitExists = (int)HitStatic;
+	}
+}
diff --git a/Engine/Plugins/Runtime/Steam/SteamAudio/Source/SteamAudio/Private/PhononGeometryComponent.cpp b/Engine/Plugins/Runtime/Steam/SteamAudio/Source/SteamAudio/Private/PhononGeometryComponent.cpp
new file mode 100644
--- /dev/null
+++ b/Engine/Plugins/Runtime/Steam/SteamAudio/Source/SteamAudio/Private/PhononGeometryComponent.cpp
@@ -0,0 +1,114 @@
+//
+// Copyright (C) Valve Corporation. All rights reserved.
+//
+
+#include "PhononGeometryComponent.h"
+#include "Engine/Engine.h"
+#include "PhononScene.h"
+#include "StaticMeshResources.h"
+#include "PhononPluginManager.h"
+
+UPhononGeometryComponent::UPhononGeometryComponent()
+	: ExportAllChildren(false)
+	, NumVertices(0)
+	, NumTriangles(0)
+{
+}
+
+
+void UPhononGeometryComponent::BeginPlay()
+{
+	Super::BeginPlay();
+
+	// Check if we are parented to a movable mesh
+	USceneComponent* Parent = GetAttachParent();
+
+	if (Parent
+		&& Parent->IsValidLowLevelFast()
+		&& Parent->Mobility == EComponentMobility::Movable
+		)
+	{
+		bHasDynamicParent = true;
+
+		// Register this component to the plugin manager as holding on to dynamic geometry
+		FAudioDeviceHandle AudioDevice = GEngine->GetMainAudioDevice();
+		if (AudioDevice)
+		{
+			for (TAudioPluginListenerPtr AudioListener : AudioDevice->PluginListeners)
+			{
+				auto SteamAudioListener = StaticCastSharedPtr<SteamAudio::FPhononPluginManager>(AudioListener);
+
+				if (SteamAudioListener && SteamAudioListener->PluginID.IsEqual(FName(TEXT("Phonon"))))
+				{
+					// Add dynamic geometry to runtime cache array
+					SteamAudioListener->AddDynamicGeometry(this);
+				}
+			}
+		}
+	}
+}
+
+
+void UPhononGeometryComponent::EndPlay(const EEndPlayReason::Type EndPlayReason)
+{
+	Super::EndPlay(EndPlayReason);
+
+	// Destroy any instanced meshes created
+	if (bHasDynamicParent)
+	{	
+		// Unregister this component from the plugin manager
+		FAudioDeviceHandle AudioDevice = GEngine->GetMainAudioDevice();
+		if (AudioDevice)
+		{
+			for (TAudioPluginListenerPtr AudioListener : AudioDevice->PluginListeners)
+			{
+				auto SteamAudioListener = StaticCastSharedPtr<SteamAudio::FPhononPluginManager>(AudioListener);
+
+				if (SteamAudioListener && SteamAudioListener->PluginID.IsEqual(FName(TEXT("Phonon"))))
+				{
+					SteamAudioListener->RemoveDynamicGeometry(this);
+				}
+			}
+		}
+	}
+
+}
+
+void UPhononGeometryComponent::UpdateStatistics()
+{
+#if WITH_EDITOR
+	if (ExportAllChildren)
+	{
+		NumTriangles = SteamAudio::GetNumTrianglesAtRoot(GetOwner());
+	}
+	else
+	{
+		NumTriangles = SteamAudio::GetNumTrianglesForStaticMesh(Cast<AStaticMeshActor>(GetOwner()));
+	}
+
+	NumVertices = NumTriangles * 3;
+#endif // WITH_EDITOR
+}
+
+#if WITH_EDITOR
+void UPhononGeometryComponent::PostEditChangeProperty(FPropertyChangedEvent& PropertyChangedEvent)
+{
+	Super::PostEditChangeProperty(PropertyChangedEvent);
+
+	FName PropertyName = (PropertyChangedEvent.Property != nullptr) ? PropertyChangedEvent.Property->GetFName() : NAME_None;
+
+	if ((PropertyName == GET_MEMBER_NAME_CHECKED(UPhononGeometryComponent, ExportAllChildren)))
+	{
+		UpdateStatistics();
+	}
+}
+#endif
+
+void UPhononGeometryComponent::OnComponentCreated()
+{
+	Super::OnComponentCreated();
+
+#if WITH_EDITOR
+	UpdateStatistics();
+#endif // WITH_EDITOR
+}
diff --git a/Engine/Plugins/Runtime/Steam/SteamAudio/Source/SteamAudio/Private/PhononGeometryComponent.h b/Engine/Plugins/Runtime/Steam/SteamAudio/Source/SteamAudio/Private/PhononGeometryComponent.h
new file mode 100644
--- /dev/null
+++ b/Engine/Plugins/Runtime/Steam/SteamAudio/Source/SteamAudio/Private/PhononGeometryComponent.h
@@ -0,0 +1,49 @@
+//
+// Copyright (C) Valve Corporation. All rights reserved.
+//
+
+#pragma once
+
+#include "Components/SceneComponent.h"
+#include "PhononCommon.h"
+#include "PhononGeometryComponent.generated.h"
+
+/**
+ * Phonon Geometry components are used to tag an actor as containing geometry relevant to acoustics calculations.
+ * Should be placed on Static Mesh actors.
+ */
+UCLASS(ClassGroup = (Audio), HideCategories = (Activation, Collision, Cooking), meta = (BlueprintSpawnableComponent))
+class UPhononGeometryComponent : public USceneComponent
+{
+	GENERATED_BODY()
+
+public:
+	UPhononGeometryComponent();
+	virtual void BeginPlay() override;
+	virtual void EndPlay(const EEndPlayReason::Type EndPlayReason) override;
+	void OnComponentCreated();
+
+#if WITH_EDITOR
+	virtual void PostEditChangeProperty(struct FPropertyChangedEvent& PropertyChangedEvent) override;
+#endif
+
+	// Whether or not to export all actors attached to this actor.
+	UPROPERTY(EditAnywhere, Category = Settings)
+	bool ExportAllChildren;
+
+	// The number of vertices exported to Steam Audio.
+	UPROPERTY(VisibleAnywhere, Category = GeometryStatistics, meta = (DisplayName = "Vertices"))
+	uint32 NumVertices;
+
+	// The number of triangles exported to Steam Audio.
+	UPROPERTY(VisibleAnywhere, Category = GeometryStatistics, meta = (DisplayName = "Triangles"))
+	uint32 NumTriangles;
+
+	// The handle to an instanced mesh if this component is attached to dynamic geometry
+	IPLhandle* InstancedMesh = nullptr;
+
+private:
+	void UpdateStatistics();
+	bool bHasDynamicParent = false;
+
+}; 
\ No newline at end of file
diff --git a/Engine/Plugins/Runtime/Steam/SteamAudio/Source/SteamAudio/Private/PhononMaterial.cpp b/Engine/Plugins/Runtime/Steam/SteamAudio/Source/SteamAudio/Private/PhononMaterial.cpp
new file mode 100644
--- /dev/null
+++ b/Engine/Plugins/Runtime/Steam/SteamAudio/Source/SteamAudio/Private/PhononMaterial.cpp
@@ -0,0 +1,113 @@
+//
+// Copyright (C) Valve Corporation. All rights reserved.
+//
+
+#include "PhononMaterial.h"
+
+static TMap<EPhononMaterial, IPLMaterial> GetMaterialPresets()
+{
+	TMap<EPhononMaterial, IPLMaterial> Presets;
+
+	IPLMaterial mat;
+	mat.scattering = 0.05f;
+
+	mat.lowFreqAbsorption = 0.10f;
+	mat.midFreqAbsorption = 0.20f;
+	mat.highFreqAbsorption = 0.30f;
+	mat.lowFreqTransmission = 0.100f;
+	mat.midFreqTransmission = 0.050f;
+	mat.highFreqTransmission = 0.030f;
+	Presets.Add(EPhononMaterial::GENERIC, mat);
+
+	mat.lowFreqAbsorption = 0.03f;
+	mat.midFreqAbsorption = 0.04f;
+	mat.highFreqAbsorption = 0.07f;
+	mat.lowFreqTransmission = 0.015f;
+	mat.midFreqTransmission = 0.015f;
+	mat.highFreqTransmission = 0.015f;
+	Presets.Add(EPhononMaterial::BRICK, mat);
+
+	mat.lowFreqAbsorption = 0.05f;
+	mat.midFreqAbsorption = 0.07f;
+	mat.highFreqAbsorption = 0.08f;
+	mat.lowFreqTransmission = 0.015f;
+	mat.midFreqTransmission = 0.002f;
+	mat.highFreqTransmission = 0.001f;
+	Presets.Add(EPhononMaterial::CONCRETE, mat);
+
+	mat.lowFreqAbsorption = 0.01f;
+	mat.midFreqAbsorption = 0.02f;
+	mat.highFreqAbsorption = 0.02f;
+	mat.lowFreqTransmission = 0.060f;
+	mat.midFreqTransmission = 0.044f;
+	mat.highFreqTransmission = 0.011f;
+	Presets.Add(EPhononMaterial::CERAMIC, mat);
+
+	mat.lowFreqAbsorption = 0.60f;
+	mat.midFreqAbsorption = 0.70f;
+	mat.highFreqAbsorption = 0.80f;
+	mat.lowFreqTransmission = 0.031f;
+	mat.midFreqTransmission = 0.012f;
+	mat.highFreqTransmission = 0.008f;
+	Presets.Add(EPhononMaterial::GRAVEL, mat);
+
+	mat.lowFreqAbsorption = 0.24f;
+	mat.midFreqAbsorption = 0.69f;
+	mat.highFreqAbsorption = 0.73f;
+	mat.lowFreqTransmission = 0.020f;
+	mat.midFreqTransmission = 0.005f;
+	mat.highFreqTransmission = 0.003f;
+	Presets.Add(EPhononMaterial::CARPET, mat);
+
+	mat.lowFreqAbsorption = 0.06f;
+	mat.midFreqAbsorption = 0.03f;
+	mat.highFreqAbsorption = 0.02f;
+	mat.lowFreqTransmission = 0.060f;
+	mat.midFreqTransmission = 0.044f;
+	mat.highFreqTransmission = 0.011f;
+	Presets.Add(EPhononMaterial::GLASS, mat);
+
+	mat.lowFreqAbsorption = 0.12f;
+	mat.midFreqAbsorption = 0.06f;
+	mat.highFreqAbsorption = 0.04f;
+	mat.lowFreqTransmission = 0.056f;
+	mat.midFreqTransmission = 0.056f;
+	mat.highFreqTransmission = 0.004f;
+	Presets.Add(EPhononMaterial::PLASTER, mat);
+
+	mat.lowFreqAbsorption = 0.11f;
+	mat.midFreqAbsorption = 0.07f;
+	mat.highFreqAbsorption = 0.06f;
+	mat.lowFreqTransmission = 0.070f;
+	mat.midFreqTransmission = 0.014f;
+	mat.highFreqTransmission = 0.005f;
+	Presets.Add(EPhononMaterial::WOOD, mat);
+
+	mat.lowFreqAbsorption = 0.20f;
+	mat.midFreqAbsorption = 0.07f;
+	mat.highFreqAbsorption = 0.06f;
+	mat.lowFreqTransmission = 0.200f;
+	mat.midFreqTransmission = 0.025f;
+	mat.highFreqTransmission = 0.010f;
+	Presets.Add(EPhononMaterial::METAL, mat);
+
+	mat.lowFreqAbsorption = 0.13f;
+	mat.midFreqAbsorption = 0.20f;
+	mat.highFreqAbsorption = 0.24f;
+	mat.lowFreqTransmission = 0.015f;
+	mat.midFreqTransmission = 0.002f;
+	mat.highFreqTransmission = 0.001f;
+	Presets.Add(EPhononMaterial::ROCK, mat);
+
+	mat.lowFreqAbsorption = 0.00f;
+	mat.midFreqAbsorption = 0.00f;
+	mat.highFreqAbsorption = 0.00f;
+	mat.lowFreqTransmission = 0.00f;
+	mat.midFreqTransmission = 0.00f;
+	mat.highFreqTransmission = 0.00f;
+	Presets.Add(EPhononMaterial::CUSTOM, mat);
+
+	return Presets;
+}
+
+TMap<EPhononMaterial, IPLMaterial> SteamAudio::MaterialPresets = GetMaterialPresets();
diff --git a/Engine/Plugins/Runtime/Steam/SteamAudio/Source/SteamAudio/Private/PhononMaterialComponent.cpp b/Engine/Plugins/Runtime/Steam/SteamAudio/Source/SteamAudio/Private/PhononMaterialComponent.cpp
new file mode 100644
--- /dev/null
+++ b/Engine/Plugins/Runtime/Steam/SteamAudio/Source/SteamAudio/Private/PhononMaterialComponent.cpp
@@ -0,0 +1,76 @@
+//
+// Copyright (C) Valve Corporation. All rights reserved.
+//
+
+#include "PhononMaterialComponent.h"
+
+UPhononMaterialComponent::UPhononMaterialComponent()
+	: MaterialIndex(0)
+	, MaterialPreset(EPhononMaterial::GENERIC)
+{
+	IPLMaterial SelectedMaterialPreset = SteamAudio::MaterialPresets[MaterialPreset];
+	LowFreqAbsorption = SelectedMaterialPreset.lowFreqAbsorption;
+	MidFreqAbsorption = SelectedMaterialPreset.midFreqAbsorption;
+	HighFreqAbsorption = SelectedMaterialPreset.highFreqAbsorption;
+	LowFreqTransmission = SelectedMaterialPreset.lowFreqTransmission;
+	MidFreqTransmission = SelectedMaterialPreset.midFreqTransmission;
+	HighFreqTransmission = SelectedMaterialPreset.highFreqTransmission;
+	Scattering = SelectedMaterialPreset.scattering;
+}
+
+IPLMaterial UPhononMaterialComponent::GetMaterialPreset() const
+{
+	IPLMaterial SelectedMaterialPreset;
+	SelectedMaterialPreset.lowFreqAbsorption = LowFreqAbsorption;
+	SelectedMaterialPreset.midFreqAbsorption = MidFreqAbsorption;
+	SelectedMaterialPreset.highFreqAbsorption = HighFreqAbsorption;
+	SelectedMaterialPreset.lowFreqTransmission = LowFreqTransmission;
+	SelectedMaterialPreset.midFreqTransmission = MidFreqTransmission;
+	SelectedMaterialPreset.highFreqTransmission = HighFreqTransmission;
+	SelectedMaterialPreset.scattering = Scattering;
+
+	return SelectedMaterialPreset;
+}
+
+#if WITH_EDITOR
+void UPhononMaterialComponent::PostEditChangeProperty(struct FPropertyChangedEvent& PropertyChangedEvent)
+{
+	Super::PostEditChangeProperty(PropertyChangedEvent);
+
+	FName PropertyName = (PropertyChangedEvent.Property != NULL) ? PropertyChangedEvent.Property->GetFName() : NAME_None;
+
+	IPLMaterial SelectedMaterialPreset = SteamAudio::MaterialPresets[MaterialPreset];
+
+	if ((PropertyName == GET_MEMBER_NAME_CHECKED(UPhononMaterialComponent, MaterialPreset)))
+	{
+		LowFreqAbsorption = SelectedMaterialPreset.lowFreqAbsorption;
+		MidFreqAbsorption = SelectedMaterialPreset.midFreqAbsorption;
+		HighFreqAbsorption = SelectedMaterialPreset.highFreqAbsorption;
+		LowFreqTransmission = SelectedMaterialPreset.lowFreqTransmission;
+		MidFreqTransmission = SelectedMaterialPreset.midFreqTransmission;
+		HighFreqTransmission = SelectedMaterialPreset.highFreqTransmission;
+		Scattering = SelectedMaterialPreset.scattering;
+	}
+}
+
+bool UPhononMaterialComponent::CanEditChange(const FProperty* InProperty) const
+{
+	const bool ParentVal = Super::CanEditChange(InProperty);
+	FName PropertyName = InProperty->GetFName();
+
+	if (PropertyName == GET_MEMBER_NAME_CHECKED(UPhononMaterialComponent, LowFreqAbsorption) ||
+		PropertyName == GET_MEMBER_NAME_CHECKED(UPhononMaterialComponent, MidFreqAbsorption) ||
+		PropertyName == GET_MEMBER_NAME_CHECKED(UPhononMaterialComponent, HighFreqAbsorption) ||
+		PropertyName == GET_MEMBER_NAME_CHECKED(UPhononMaterialComponent, LowFreqTransmission) ||
+		PropertyName == GET_MEMBER_NAME_CHECKED(UPhononMaterialComponent, MidFreqTransmission) ||
+		PropertyName == GET_MEMBER_NAME_CHECKED(UPhononMaterialComponent, HighFreqTransmission) ||
+		PropertyName == GET_MEMBER_NAME_CHECKED(UPhononMaterialComponent, Scattering))
+	{
+		return ParentVal && (MaterialPreset == EPhononMaterial::CUSTOM);
+	}
+	else
+	{
+		return ParentVal;
+	}
+}
+#endif
diff --git a/Engine/Plugins/Runtime/Steam/SteamAudio/Source/SteamAudio/Private/PhononMaterialComponent.h b/Engine/Plugins/Runtime/Steam/SteamAudio/Source/SteamAudio/Private/PhononMaterialComponent.h
new file mode 100644
--- /dev/null
+++ b/Engine/Plugins/Runtime/Steam/SteamAudio/Source/SteamAudio/Private/PhononMaterialComponent.h
@@ -0,0 +1,65 @@
+//
+// Copyright (C) Valve Corporation. All rights reserved.
+//
+
+#pragma once
+
+#include "PhononMaterial.h"
+#include "Components/ActorComponent.h"
+#include "PhononMaterialComponent.generated.h"
+
+/**
+ * Phonon Material components are used to customize an actor's acoustic properties. Only valid on actors that also
+ * have a Phonon Geometry component.
+ */
+UCLASS(ClassGroup = (Audio), HideCategories = (Activation, Collision, Cooking), meta = (BlueprintSpawnableComponent))
+class UPhononMaterialComponent : public UActorComponent
+{
+	GENERATED_BODY()
+
+public:
+	UPhononMaterialComponent();
+
+	IPLMaterial GetMaterialPreset() const;
+
+#if WITH_EDITOR
+	virtual void PostEditChangeProperty(struct FPropertyChangedEvent& PropertyChangedEvent) override;
+	virtual bool CanEditChange(const FProperty* InProperty) const override;
+#endif
+
+	UPROPERTY()
+	int32 MaterialIndex;
+
+	// Choose from a variety of preset physical materials, or choose Custom to specify values manually.
+	UPROPERTY(EditAnywhere, Category = Settings)
+	EPhononMaterial MaterialPreset;
+
+	// How much this material absorbs low frequency sound.
+	UPROPERTY(EditAnywhere, Category = Settings, meta = (ClampMin = "0.0", ClampMax = "1.0", UIMin = "0.0", UIMax = "1.0"))
+	float LowFreqAbsorption;
+
+	// How much this material absorbs mid frequency sound.
+	UPROPERTY(EditAnywhere, Category = Settings, meta = (ClampMin = "0.0", ClampMax = "1.0", UIMin = "0.0", UIMax = "1.0"))
+	float MidFreqAbsorption;
+
+	// How much this material absorbs high frequency sound.
+	UPROPERTY(EditAnywhere, Category = Settings, meta = (ClampMin = "0.0", ClampMax = "1.0", UIMin = "0.0", UIMax = "1.0"))
+	float HighFreqAbsorption;
+
+	// How much this material transmits low frequency sound.
+	UPROPERTY(EditAnywhere, Category = Settings, meta = (ClampMin = "0.0", ClampMax = "1.0", UIMin = "0.0", UIMax = "1.0"))
+	float LowFreqTransmission;
+
+	// How much this material transmits mid frequency sound.
+	UPROPERTY(EditAnywhere, Category = Settings, meta = (ClampMin = "0.0", ClampMax = "1.0", UIMin = "0.0", UIMax = "1.0"))
+	float MidFreqTransmission;
+
+	// How much this material transmits high frequency sound.
+	UPROPERTY(EditAnywhere, Category = Settings, meta = (ClampMin = "0.0", ClampMax = "1.0", UIMin = "0.0", UIMax = "1.0"))
+	float HighFreqTransmission;
+
+	// Specifies how "rough" the surface is. Surfaces with a high scattering value randomly reflect sound in all directions;
+	// surfaces with a low scattering value reflect sound in a mirror-like manner.
+	UPROPERTY(EditAnywhere, Category = Settings, meta = (ClampMin = "0.0", ClampMax = "1.0", UIMin = "0.0", UIMax = "1.0"))
+	float Scattering;
+};
\ No newline at end of file
diff --git a/Engine/Plugins/Runtime/Steam/SteamAudio/Source/SteamAudio/Private/PhononOcclusion.cpp b/Engine/Plugins/Runtime/Steam/SteamAudio/Source/SteamAudio/Private/PhononOcclusion.cpp
new file mode 100644
--- /dev/null
+++ b/Engine/Plugins/Runtime/Steam/SteamAudio/Source/SteamAudio/Private/PhononOcclusion.cpp
@@ -0,0 +1,184 @@
+//
+// Copyright (C) Valve Corporation. All rights reserved.
+//
+
+#include "PhononOcclusion.h"
+#include "PhononOcclusionSourceSettings.h"
+#include "PhononCommon.h"
+#include "SteamAudioModule.h"
+#include "Misc/ScopeLock.h"
+#include "SteamAudioEnvironment.h"
+
+//==================================================================================================================================================
+// FPhononOcclusion
+//==================================================================================================================================================
+
+namespace SteamAudio
+{
+	FPhononOcclusion::FPhononOcclusion()
+		: Environment(nullptr)
+	{
+		InputAudioFormat.channelLayout = IPL_CHANNELLAYOUT_MONO;
+		InputAudioFormat.channelLayoutType = IPL_CHANNELLAYOUTTYPE_SPEAKERS;
+		InputAudioFormat.channelOrder = IPL_CHANNELORDER_INTERLEAVED;
+		InputAudioFormat.numSpeakers = 1;
+		InputAudioFormat.speakerDirections = nullptr;
+		InputAudioFormat.ambisonicsOrder = -1;
+		InputAudioFormat.ambisonicsNormalization = IPL_AMBISONICSNORMALIZATION_N3D;
+		InputAudioFormat.ambisonicsOrdering = IPL_AMBISONICSORDERING_ACN;
+
+		OutputAudioFormat.channelLayout = IPL_CHANNELLAYOUT_MONO;
+		OutputAudioFormat.channelLayoutType = IPL_CHANNELLAYOUTTYPE_SPEAKERS;
+		OutputAudioFormat.channelOrder = IPL_CHANNELORDER_INTERLEAVED;
+		OutputAudioFormat.numSpeakers = 1;
+		OutputAudioFormat.speakerDirections = nullptr;
+		OutputAudioFormat.ambisonicsOrder = -1;
+		OutputAudioFormat.ambisonicsNormalization = IPL_AMBISONICSNORMALIZATION_N3D;
+		OutputAudioFormat.ambisonicsOrdering = IPL_AMBISONICSORDERING_ACN;
+	}
+
+	FPhononOcclusion::~FPhononOcclusion()
+	{
+	}
+
+	void FPhononOcclusion::Initialize(const FAudioPluginInitializationParams InitializationParams)
+	{
+		DirectSoundSources.AddDefaulted(InitializationParams.NumSources);
+
+		for (FDirectSoundSource& DirectSoundSource : DirectSoundSources)
+		{
+			DirectSoundSource.InBuffer.format = InputAudioFormat;
+			DirectSoundSource.InBuffer.numSamples = InitializationParams.BufferLength;
+			DirectSoundSource.InBuffer.interleavedBuffer = nullptr;
+			DirectSoundSource.InBuffer.deinterleavedBuffer = nullptr;
+
+			DirectSoundSource.OutBuffer.format = OutputAudioFormat;
+			DirectSoundSource.OutBuffer.numSamples = InitializationParams.BufferLength;
+			DirectSoundSource.OutBuffer.interleavedBuffer = nullptr;
+			DirectSoundSource.OutBuffer.deinterleavedBuffer = nullptr;
+		}
+	}
+
+	void FPhononOcclusion::OnInitSource(const uint32 SourceId, const FName& AudioComponentUserId, const uint32 NumChannels, UOcclusionPluginSourceSettingsBase* InSettings)
+	{
+		if (!Environment || !Environment->GetEnvironmentalRenderer())
+		{
+			UE_LOG(LogSteamAudio, Error, TEXT("Unable to find environmental renderer for occlusion. Audio will not be occluded. Make sure to export the scene."));
+			return;
+		}
+
+		UE_LOG(LogSteamAudio, Log, TEXT("Creating occlusion effect."));
+
+		UPhononOcclusionSourceSettings* OcclusionSettings = CastChecked<UPhononOcclusionSourceSettings>(InSettings);
+		DirectSoundSources[SourceId].bDirectAttenuation = OcclusionSettings->DirectAttenuation;
+		DirectSoundSources[SourceId].bAirAbsorption = OcclusionSettings->AirAbsorption;
+		DirectSoundSources[SourceId].DirectOcclusionMethod = OcclusionSettings->DirectOcclusionMethod;
+		DirectSoundSources[SourceId].DirectOcclusionMode = OcclusionSettings->DirectOcclusionMode;
+		DirectSoundSources[SourceId].Radius = OcclusionSettings->DirectOcclusionSourceRadius;
+		DirectSoundSources[SourceId].SourceData.directivity = IPLDirectivity{ 0.0f, 0.0f, nullptr, nullptr };
+
+		InputAudioFormat.numSpeakers = OutputAudioFormat.numSpeakers = NumChannels;
+		switch (NumChannels)
+		{
+			case 1: InputAudioFormat.channelLayout = OutputAudioFormat.channelLayout = IPL_CHANNELLAYOUT_MONO; break;
+			case 2: InputAudioFormat.channelLayout = OutputAudioFormat.channelLayout = IPL_CHANNELLAYOUT_STEREO; break;
+			case 4: InputAudioFormat.channelLayout = OutputAudioFormat.channelLayout = IPL_CHANNELLAYOUT_QUADRAPHONIC; break;
+			case 6: InputAudioFormat.channelLayout = OutputAudioFormat.channelLayout = IPL_CHANNELLAYOUT_FIVEPOINTONE; break;
+			case 8: InputAudioFormat.channelLayout = OutputAudioFormat.channelLayout = IPL_CHANNELLAYOUT_SEVENPOINTONE; break;
+		}
+
+		DirectSoundSources[SourceId].InBuffer.format = InputAudioFormat;
+		DirectSoundSources[SourceId].OutBuffer.format = OutputAudioFormat;
+
+		iplCreateDirectSoundEffect(Environment->GetEnvironmentalRenderer(), InputAudioFormat, OutputAudioFormat, &(DirectSoundSources[SourceId].DirectSoundEffect));
+	}
+
+	void FPhononOcclusion::OnReleaseSource(const uint32 SourceId)
+	{
+		UE_LOG(LogSteamAudio, Log, TEXT("Destroying occlusion effect."));
+
+		iplDestroyDirectSoundEffect(&(DirectSoundSources[SourceId].DirectSoundEffect));
+	}
+
+	void FPhononOcclusion::ProcessAudio(const FAudioPluginSourceInputData& InputData, FAudioPluginSourceOutputData& OutputData)
+	{
+		FDirectSoundSource& DirectSoundSource = DirectSoundSources[InputData.SourceId];
+
+		if (!Environment || !Environment->GetEnvironmentalRenderer())
+		{
+			FMemory::Memcpy(OutputData.AudioBuffer.GetData(), InputData.AudioBuffer->GetData(), InputData.AudioBuffer->Num() * sizeof(float));
+			return;
+		}
+
+		DirectSoundSource.InBuffer.interleavedBuffer = InputData.AudioBuffer->GetData();
+		DirectSoundSource.OutBuffer.interleavedBuffer = OutputData.AudioBuffer.GetData();
+
+		{
+			FScopeLock Lock(&DirectSoundSources[InputData.SourceId].CriticalSection);
+			DirectSoundSource.SourceData.position = SteamAudio::UnrealToPhononIPLVector3(InputData.SpatializationParams->EmitterWorldPosition);
+			DirectSoundSource.SourceData.ahead = SteamAudio::UnrealToPhononIPLVector3(InputData.SpatializationParams->EmitterWorldRotation * FVector::ForwardVector);
+			DirectSoundSource.SourceData.right = SteamAudio::UnrealToPhononIPLVector3(InputData.SpatializationParams->EmitterWorldRotation * FVector::RightVector);
+			DirectSoundSource.SourceData.up = SteamAudio::UnrealToPhononIPLVector3(InputData.SpatializationParams->EmitterWorldRotation * FVector::UpVector);
+			DirectSoundSource.bNeedsUpdate = true;
+		}
+
+		IPLDirectSoundEffectOptions DirectSoundEffectOptions;
+		DirectSoundEffectOptions.applyAirAbsorption = static_cast<IPLbool>(DirectSoundSources[InputData.SourceId].bAirAbsorption);
+		DirectSoundEffectOptions.applyDistanceAttenuation = static_cast<IPLbool>(DirectSoundSources[InputData.SourceId].bDirectAttenuation);
+		DirectSoundEffectOptions.directOcclusionMode = static_cast<IPLDirectOcclusionMode>(DirectSoundSources[InputData.SourceId].DirectOcclusionMode);
+
+		iplApplyDirectSoundEffect(DirectSoundSource.DirectSoundEffect, DirectSoundSource.InBuffer, DirectSoundSource.DirectSoundPath, DirectSoundEffectOptions, DirectSoundSource.OutBuffer);
+	}
+
+	void FPhononOcclusion::UpdateDirectSoundSources(const FVector& ListenerPosition, const FVector& ListenerForward, const FVector& ListenerUp, const FVector& ListenerRight)
+	{
+		if (!Environment || !Environment->GetEnvironmentalRenderer() || !Environment->GetEnvironmentCriticalSectionHandle())
+		{
+			return;
+		}
+
+		FScopeLock EnvironmentLock(Environment->GetEnvironmentCriticalSectionHandle());
+
+		for (FDirectSoundSource& DirectSoundSource : DirectSoundSources)
+		{
+			FScopeLock DirectSourceLock(&DirectSoundSource.CriticalSection);
+
+			if (DirectSoundSource.bNeedsUpdate)
+			{
+				IPLDirectSoundPath DirectSoundPath = iplGetDirectSoundPath(Environment->GetEnvironment(), SteamAudio::UnrealToPhononIPLVector3(ListenerPosition),
+					SteamAudio::UnrealToPhononIPLVector3(ListenerForward, false), SteamAudio::UnrealToPhononIPLVector3(ListenerUp, false),
+					DirectSoundSource.SourceData, DirectSoundSource.Radius * SteamAudio::SCALEFACTOR,
+					static_cast<IPLDirectOcclusionMode>(DirectSoundSource.DirectOcclusionMode),
+					static_cast<IPLDirectOcclusionMethod>(DirectSoundSource.DirectOcclusionMethod));
+
+				DirectSoundSource.DirectSoundPath = DirectSoundPath;
+				DirectSoundSource.bNeedsUpdate = false;
+			}
+		}
+	}
+
+	void FPhononOcclusion::SetEnvironment(FEnvironment* InEnvironment)
+	{
+		Environment = InEnvironment;
+	}
+
+}
+
+//==================================================================================================================================================
+// FDirectSoundSource
+//==================================================================================================================================================
+
+namespace SteamAudio
+{
+	FDirectSoundSource::FDirectSoundSource()
+		: DirectSoundEffect(nullptr)
+		, DirectOcclusionMethod(EIplDirectOcclusionMethod::RAYCAST)
+		, DirectOcclusionMode(EIplDirectOcclusionMode::NONE)
+		, bDirectAttenuation(false)
+		, bAirAbsorption(false)
+		, bNeedsUpdate(false)
+	{
+		memset(&DirectSoundPath, 0, sizeof(DirectSoundPath));
+		memset(&SourceData, 0, sizeof(SourceData));
+	}
+}
diff --git a/Engine/Plugins/Runtime/Steam/SteamAudio/Source/SteamAudio/Private/PhononOcclusion.h b/Engine/Plugins/Runtime/Steam/SteamAudio/Source/SteamAudio/Private/PhononOcclusion.h
new file mode 100644
--- /dev/null
+++ b/Engine/Plugins/Runtime/Steam/SteamAudio/Source/SteamAudio/Private/PhononOcclusion.h
@@ -0,0 +1,72 @@
+//
+// Copyright (C) Valve Corporation. All rights reserved.
+//
+
+#pragma once
+
+#include "IAudioExtensionPlugin.h"
+#include "PhononCommon.h"
+#include "phonon.h"
+
+class UOcclusionPluginSourceSettingsBase;
+
+namespace SteamAudio
+{
+	class FEnvironment;
+
+	struct FDirectSoundSource
+	{
+		FDirectSoundSource();
+
+		FCriticalSection CriticalSection;
+		IPLDirectSoundPath DirectSoundPath;
+		IPLhandle DirectSoundEffect;
+		EIplDirectOcclusionMethod DirectOcclusionMethod;
+		EIplDirectOcclusionMode DirectOcclusionMode;
+		IPLAudioBuffer InBuffer;
+		IPLAudioBuffer OutBuffer;
+		IPLSource SourceData;
+		float Radius;
+		float DipolePower;
+		float DipoleWeight;
+		bool bDirectAttenuation;
+		bool bAirAbsorption;
+		bool bNeedsUpdate;
+	};
+
+	/************************************************************************/
+	/* FPhononOcclusion                                                     */
+	/* Scene-dependent audio occlusion plugin. Receives updates from        */
+	/* an FPhononPluginManager on the game thread on player position and    */
+	/* geometry, and performs geometry-aware filtering of the direct path   */
+	/* of an audio source.                                                  */
+	/************************************************************************/
+	class FPhononOcclusion : public IAudioOcclusion
+	{
+	public:
+		FPhononOcclusion();
+		~FPhononOcclusion();
+
+		//~ Begin IAudioOcclusion
+		virtual void Initialize(const FAudioPluginInitializationParams InitializationParams) override;
+		virtual void OnInitSource(const uint32 SourceId, const FName& AudioComponentUserId, const uint32 NumChannels, UOcclusionPluginSourceSettingsBase* InSettings) override;
+		virtual void OnReleaseSource(const uint32 SourceId) override;
+		virtual void ProcessAudio(const FAudioPluginSourceInputData& InputData, FAudioPluginSourceOutputData& OutputData) override;
+		//~ End IAudioOcclusion
+
+		// Receives updates on listener positions from the game thread using this function call.
+		void UpdateDirectSoundSources(const FVector& ListenerPosition, const FVector& ListenerForward, const FVector& ListenerUp, const FVector& ListenerRight);
+
+		// Sets up handle to the environment handle owned by the FPhononPluginManager.
+		void SetEnvironment(FEnvironment* InEnvironment);
+
+	private:
+		IPLAudioFormat InputAudioFormat;
+		IPLAudioFormat OutputAudioFormat;
+
+		FEnvironment* Environment;
+
+		// Cached array of direct sound sources to be occluded.
+		TArray<FDirectSoundSource> DirectSoundSources;
+	};
+}
diff --git a/Engine/Plugins/Runtime/Steam/SteamAudio/Source/SteamAudio/Private/PhononOcclusionSourceSettings.cpp b/Engine/Plugins/Runtime/Steam/SteamAudio/Source/SteamAudio/Private/PhononOcclusionSourceSettings.cpp
new file mode 100644
--- /dev/null
+++ b/Engine/Plugins/Runtime/Steam/SteamAudio/Source/SteamAudio/Private/PhononOcclusionSourceSettings.cpp
@@ -0,0 +1,31 @@
+//
+// Copyright (C) Valve Corporation. All rights reserved.
+//
+
+#include "PhononOcclusionSourceSettings.h"
+
+UPhononOcclusionSourceSettings::UPhononOcclusionSourceSettings()
+	: DirectOcclusionMode(EIplDirectOcclusionMode::NONE)
+	, DirectOcclusionMethod(EIplDirectOcclusionMethod::RAYCAST)
+	, DirectOcclusionSourceRadius(100.0f)
+	, DirectAttenuation(true)
+	, AirAbsorption(true)
+{}
+
+#if WITH_EDITOR
+bool UPhononOcclusionSourceSettings::CanEditChange(const FProperty* InProperty) const
+{
+	const bool ParentVal = Super::CanEditChange(InProperty);
+
+	if ((InProperty->GetFName() == GET_MEMBER_NAME_CHECKED(UPhononOcclusionSourceSettings, DirectOcclusionSourceRadius)))
+	{
+		return ParentVal && DirectOcclusionMode != EIplDirectOcclusionMode::NONE && DirectOcclusionMethod == EIplDirectOcclusionMethod::VOLUMETRIC;
+	}
+	else if ((InProperty->GetFName() == GET_MEMBER_NAME_CHECKED(UPhononOcclusionSourceSettings, DirectOcclusionMethod)))
+	{
+		return ParentVal && DirectOcclusionMode != EIplDirectOcclusionMode::NONE;
+	}
+
+	return ParentVal;
+}
+#endif
diff --git a/Engine/Plugins/Runtime/Steam/SteamAudio/Source/SteamAudio/Private/PhononPluginManager.cpp b/Engine/Plugins/Runtime/Steam/SteamAudio/Source/SteamAudio/Private/PhononPluginManager.cpp
new file mode 100644
--- /dev/null
+++ b/Engine/Plugins/Runtime/Steam/SteamAudio/Source/SteamAudio/Private/PhononPluginManager.cpp
@@ -0,0 +1,310 @@
+//
+// Copyright (C) Valve Corporation. All rights reserved.
+//
+
+#include "PhononPluginManager.h"
+#include "SteamAudioModule.h"
+#include "AudioDevice.h"
+#include "HAL/FileManager.h"
+#include "DrawDebugHelpers.h"
+#include "Engine/World.h"
+#include "Misc/Paths.h"
+
+namespace SteamAudio
+{
+	const FName FPhononPluginManager::PluginID = FName("Phonon");
+	TArray<FDynamicGeometryMap> FPhononPluginManager::DynamicGeometry;
+	TMap<FString, IPLhandle> FPhononPluginManager::DynamicSceneMap;
+	UWorld* FPhononPluginManager::CurrentLevel;
+		
+	FPhononPluginManager::FPhononPluginManager()
+		: bEnvironmentInitialized(false)
+		, ReverbPtr(nullptr)
+		, OcclusionPtr(nullptr)
+	{
+	}
+
+	FPhononPluginManager::~FPhononPluginManager()
+	{
+		// Perform cleanup here instead of in OnListenerShutdown, because plugins will still be active and may be using them
+		if (bEnvironmentInitialized)
+		{
+			Environment.Shutdown();
+			bEnvironmentInitialized = false;
+		}
+	}
+
+	void FPhononPluginManager::InitializeEnvironment(FAudioDevice* AudioDevice, UWorld* ListenerWorld)
+	{
+		if (ListenerWorld->WorldType == EWorldType::Editor || bEnvironmentInitialized)
+		{
+			return;
+		}
+
+		UE_LOG(LogSteamAudio, Log, TEXT("Initializing environment for [%s] world type [%i]."), *ListenerWorld->GetMapName(), (unsigned int)ListenerWorld->WorldType.GetValue());
+
+		CurrentLevel = ListenerWorld;
+		bool bIsUsingOcclusion = IsUsingSteamAudioPlugin(EAudioPlugin::OCCLUSION);
+		bool bIsUsingReverb = IsUsingSteamAudioPlugin(EAudioPlugin::REVERB);
+
+		// Get level name
+		FString CurrentLevelName;
+		CurrentLevelName = CurrentLevel->GetName();
+
+		if (bIsUsingOcclusion || bIsUsingReverb)
+		{
+			if (Environment.Initialize(ListenerWorld, AudioDevice, &CurrentLevelName))
+			{
+				FScopeLock EnvironmentLock(Environment.GetEnvironmentCriticalSectionHandle());
+
+				UE_LOG(LogSteamAudio, Log, TEXT("Environment initialization successful for level %s."), *CurrentLevelName);
+
+				if (bIsUsingReverb)
+				{
+					ReverbPtr = static_cast<FPhononReverb*>(AudioDevice->ReverbPluginInterface.Get());
+					ReverbPtr->SetEnvironment(&Environment);
+					ReverbPtr->CreateReverbEffect();
+				}
+
+				if (bIsUsingOcclusion)
+				{
+					OcclusionPtr = static_cast<FPhononOcclusion*>(AudioDevice->OcclusionInterface.Get());
+					OcclusionPtr->SetEnvironment(&Environment);
+				}
+				
+				// Load dynamic scenes for the current level (if any)
+				TArray<FString> DynamicSceneFiles;
+				FString ScenePath = IFileManager::Get().ConvertToAbsolutePathForExternalAppForRead(*(DynamicRuntimePath / CurrentLevelName + "_*.phononscene"));
+
+				if (FPaths::DirectoryExists(DynamicRuntimePath))
+				{
+					IFileManager& FileManager = IFileManager::Get();
+					FileManager.FindFiles(DynamicSceneFiles, *ScenePath, true, false);
+
+					for (auto DynamicSceneFile : DynamicSceneFiles)
+					{
+						IPLhandle DynamicPhononScene = nullptr;
+						FPhononSceneInfo PhononSceneInfo;
+						FString DynamicSceneFileName = FPaths::GetBaseFilename(DynamicSceneFile);
+
+						// Get Dynamic Scene File
+						if (!LoadDynamicSceneFromDisk(CurrentLevel, Environment.GetComputeDevice(), Environment.GetSimulationSettings(), &DynamicPhononScene, PhononSceneInfo, DynamicSceneFileName))
+						{
+							UE_LOG(LogSteamAudio, Warning, TEXT("Unable to load dynamic scene file [%s]. Do a Scene Export and ensure your ray tracer is supported (Embree)."), *DynamicSceneFileName);
+							continue;
+						}
+
+						// Check if this scene have dynamic meshes to be added
+						bool bFoundInMap = false;
+						for (int32 i = 0; i < DynamicGeometry.Num(); i++)
+						{
+							if (DynamicGeometry[i].ID.Equals(DynamicSceneFileName)
+								&& DynamicGeometry[i].DynamicGeometryComponent
+								&& DynamicGeometry[i].DynamicGeometryComponent->IsValidLowLevelFast()
+								)
+							{
+								// Get 4x4 Matrix (1D array) from world transform
+								float Matrix[16];
+								auto DynGeoTransform = DynamicGeometry[i].DynamicGeometryComponent->GetAttachParent()->GetComponentTransform();
+								
+								GetMatrixForTransform(DynGeoTransform, Matrix, true, false);
+								IPLMatrix4x4 PhononMatrix = GetIPLMatrix(Matrix);
+								//DrawDebugSphere(ListenerWorld, DynamicGeometry[i].DynamicGeometryComponent->GetOwner()->GetRootComponent()->GetComponentTransform().GetTranslation(), 20.f, 4, FColor::Green, 9999.f);
+
+								// Create SteamAudio Instanced Mesh
+								IPLhandle InstancedMesh = nullptr;
+								IPLerror ErrCreateMesh = iplCreateInstancedMesh(Environment.GetScene(), DynamicPhononScene, PhononMatrix, &InstancedMesh);
+
+								if (ErrCreateMesh == IPL_STATUS_SUCCESS)
+								{
+									// Add Instanced Mesh to scene
+									iplAddInstancedMesh(Environment.GetScene(), InstancedMesh);
+
+									// Add references to map
+									DynamicGeometry[i].DynamicScene = DynamicPhononScene;
+									DynamicGeometry[i].DynamicGeometry = InstancedMesh;
+									UE_LOG(LogSteamAudio, Log, TEXT("Created instanced mesh [%s]"), *DynamicSceneFileName);
+								}
+								else
+								{
+									UE_LOG(LogSteamAudio, Warning, TEXT("Unable to create instanced mesh for [%s]"), *DynamicSceneFileName);
+								}
+
+								bFoundInMap = true;
+								break;
+							}
+						}
+
+						if (!bFoundInMap)
+						{
+							DynamicSceneMap.Add(DynamicSceneFileName, DynamicPhononScene);
+							UE_LOG(LogSteamAudio, Warning, TEXT("Loaded dynamic scene file [%s] but no dynamic geometry found in this level."), *DynamicSceneFileName);
+						}
+					}
+				}
+
+				bEnvironmentInitialized = true;
+			}
+			else
+			{
+				UE_LOG(LogSteamAudio, Warning, TEXT("Environment initialization unsuccessful for Level %s."), *CurrentLevelName);
+			}
+		}
+	}
+
+	void FPhononPluginManager::OnListenerInitialize(FAudioDevice* AudioDevice, UWorld* ListenerWorld)
+	{
+		InitializeEnvironment(AudioDevice, ListenerWorld);
+	}
+
+	void FPhononPluginManager::OnListenerUpdated(FAudioDevice* AudioDevice, const int32 ViewportIndex, const FTransform& ListenerTransform, const float InDeltaSeconds)
+	{
+		if (!bEnvironmentInitialized)
+		{
+			return;
+		}
+
+		FVector Position = ListenerTransform.GetLocation();
+		FVector Forward = ListenerTransform.GetUnitAxis(EAxis::X);
+		FVector Up = ListenerTransform.GetUnitAxis(EAxis::Z);
+		FVector Right = ListenerTransform.GetUnitAxis(EAxis::Y);
+
+		if (OcclusionPtr)
+		{
+			OcclusionPtr->UpdateDirectSoundSources(Position, Forward, Up, Right);
+		}
+
+		if (ReverbPtr)
+		{
+			ReverbPtr->UpdateListener(Position, Forward, Up, Right);
+		}
+	}
+
+	void FPhononPluginManager::OnListenerShutdown(FAudioDevice* AudioDevice)
+	{
+		FSteamAudioModule* Module = &FModuleManager::GetModuleChecked<FSteamAudioModule>("SteamAudio");
+		if (Module != nullptr)
+		{
+			Module->UnregisterAudioDevice(AudioDevice);
+		}
+	}
+	
+	void FPhononPluginManager::OnTick(UWorld* InWorld, const int32 ViewportIndex, const FTransform& ListenerTransform, const float InDeltaSeconds)
+	{
+		// Update any dynamic geometry
+		if (DynamicGeometry.Num() > 0)
+		{
+			// Update the transforms of the dynamic geometry
+			for (auto Geometry : DynamicGeometry)
+			{
+				// Check if we both have a dynamic mesh and dynamic scene
+				if (Geometry.DynamicGeometry
+					&& Geometry.DynamicScene
+					&& Geometry.DynamicGeometryComponent
+					&& Geometry.DynamicGeometryComponent->IsValidLowLevelFast()
+					)
+				{
+					// Get 4x4 Matrix (1D array) from world transform
+					float Matrix[16];
+					FTransform DynGeoTransform = Geometry.DynamicGeometryComponent->GetAttachParent()->GetComponentTransform();
+
+					GetMatrixForTransform(DynGeoTransform, Matrix, true, false);
+					IPLMatrix4x4 PhononMatrix = GetIPLMatrix(Matrix);
+					//DrawDebugSphere(InWorld, DynGeoTransform.GetTranslation(), 20.f, 4, FColor::Red);
+
+					// Update instanced mesh transform
+					iplUpdateInstancedMeshTransform(Geometry.DynamicGeometry, PhononMatrix);
+
+					// Commit changes to the main scene
+					iplCommitScene(Environment.GetScene());
+				}
+			}
+		}
+	}
+
+	void FPhononPluginManager::OnWorldChanged(FAudioDevice* AudioDevice, UWorld* ListenerWorld)
+	{
+		UE_LOG(LogSteamAudio, Log, TEXT("World changed to %s. Reinitializing environment."), *ListenerWorld->GetMapName());
+
+		if (bEnvironmentInitialized)
+		{
+			Environment.Shutdown();
+			bEnvironmentInitialized = false;
+		}
+
+		InitializeEnvironment(AudioDevice, ListenerWorld);
+	}
+
+	bool FPhononPluginManager::IsUsingSteamAudioPlugin(EAudioPlugin PluginType)
+	{
+		FSteamAudioModule* Module = &FModuleManager::GetModuleChecked<FSteamAudioModule>("SteamAudio");
+
+		// If we can't get the module from the module manager, then we don't have any of these plugins loaded.
+		if (Module == nullptr)
+		{
+			return false;
+		}
+
+		FString SteamPluginName = Module->GetPluginFactory(PluginType)->GetDisplayName();
+		FString CurrentPluginName = AudioPluginUtilities::GetDesiredPluginName(PluginType);
+		return CurrentPluginName.Equals(SteamPluginName);
+	}
+	
+	void FPhononPluginManager::AddDynamicGeometry(UPhononGeometryComponent* PhononGeometryComponent)
+	{
+		check(PhononGeometryComponent);
+
+		// Check if there's a valid world and the geometry component being registered is valid
+		if (PhononGeometryComponent->IsValidLowLevel())
+		{
+			FString DynamicGeometryID = PhononGeometryComponent->GetWorld()->GetName() + "_" + PhononGeometryComponent->GetOwner()->GetName() + "_" + PhononGeometryComponent->GetAttachParent()->GetName();
+			DynamicGeometry.Add({ DynamicGeometryID, PhononGeometryComponent, nullptr, nullptr });
+			UE_LOG(LogSteamAudio, Log, TEXT("Dynamic Geometry Added [%s] Current Count: %i"), *DynamicGeometryID, DynamicGeometry.Num());
+		}
+	}
+
+
+	void FPhononPluginManager::RemoveDynamicGeometry(UPhononGeometryComponent* PhononGeometryComponent)
+	{
+		check(PhononGeometryComponent);
+
+		// Check if there's a valid world and the geometry component being registered is valid
+		if (PhononGeometryComponent->IsValidLowLevel())
+		{
+			FString DynamicGeometryID = PhononGeometryComponent->GetWorld()->GetName() + "_" + PhononGeometryComponent->GetOwner()->GetName() + "_" + PhononGeometryComponent->GetAttachParent()->GetName();
+
+			// Check if this scene have dynamic meshes to be added
+			bool bFoundInMap = false;
+			for (int32 i = 0; i < DynamicGeometry.Num(); i++)
+			{
+				if (DynamicGeometry[i].ID.Equals(DynamicGeometryID)
+					&& DynamicGeometry[i].DynamicGeometryComponent
+					&& DynamicGeometry[i].DynamicGeometryComponent->IsValidLowLevelFast()
+					)
+				{
+					// Delete ipl instances
+					if (DynamicGeometry[i].DynamicGeometry)
+					{
+						iplDestroyStaticMesh(&(DynamicGeometry[i].DynamicGeometry));
+					}
+
+					if (DynamicGeometry[i].DynamicScene)
+					{
+						iplDestroyScene(&(DynamicGeometry[i].DynamicScene));
+					}
+
+					// Remove map entry
+					DynamicGeometry.RemoveAt(i);
+				}
+			}
+
+		}
+
+	}
+	
+	FEnvironment& FPhononPluginManager::GetPhononEnvironment()
+	{
+		return Environment;
+	}
+	
+}
diff --git a/Engine/Plugins/Runtime/Steam/SteamAudio/Source/SteamAudio/Private/PhononPluginManager.h b/Engine/Plugins/Runtime/Steam/SteamAudio/Source/SteamAudio/Private/PhononPluginManager.h
new file mode 100644
--- /dev/null
+++ b/Engine/Plugins/Runtime/Steam/SteamAudio/Source/SteamAudio/Private/PhononPluginManager.h
@@ -0,0 +1,74 @@
+//
+// Copyright (C) Valve Corporation. All rights reserved.
+//
+
+#pragma once
+
+#include "SteamAudioModule.h"
+#include "PhononSpatialization.h"
+#include "PhononOcclusion.h"
+#include "PhononReverb.h"
+#include "SteamAudioEnvironment.h"
+#include "PhononCommon.h"
+#include "AudioPluginUtilities.h"
+#include "PhononGeometryComponent.h"
+
+namespace SteamAudio
+{
+	/************************************************************************/
+	/* FPhononPluginManager                                                 */
+	/* This ListenerObserver owns the Steam Audio environment, and          */
+	/* dispatches information to the Steam Audio reverb and occlusion       */
+	/* plugins.                                                             */
+	/************************************************************************/
+	class FPhononPluginManager : public IAudioPluginListener
+	{
+	public:
+		FPhononPluginManager();
+		FPhononPluginManager(AActor* SubLevelRef);
+		~FPhononPluginManager();
+
+		//~ Begin IAudioPluginListener
+		virtual void OnListenerInitialize(FAudioDevice* AudioDevice, UWorld* ListenerWorld) override;
+		virtual void OnTick(UWorld* InWorld, const int32 ViewportIndex, const FTransform& ListenerTransform, const float InDeltaSeconds) override;
+		virtual void OnListenerUpdated(FAudioDevice* AudioDevice, const int32 ViewportIndex, const FTransform& ListenerTransform, const float InDeltaSeconds) override;
+		virtual void OnListenerShutdown(FAudioDevice* AudioDevice) override;
+		virtual void OnWorldChanged(FAudioDevice* AudioDevice, UWorld* ListenerWorld) override;
+		//~ End IAudioPluginListener
+
+		/** Cache of dynamic geometry currently active in the level. UPROPERTY to prevent GC */
+		UPROPERTY()
+		static TArray<FDynamicGeometryMap> DynamicGeometry;
+
+		/** Cache of dynamic geometry scenes currently active in the level. UPROPERTY to prevent GC */
+		UPROPERTY()
+		static TMap<FString, IPLhandle> DynamicSceneMap;
+
+		/** Add dynamic geometry to runtime cache */
+		static void AddDynamicGeometry(UPhononGeometryComponent* PhononGeometryComponent);
+
+		/** Remove dynamic geometry to runtime cache */
+		static void RemoveDynamicGeometry(UPhononGeometryComponent* PhononGeometryComponent);
+
+		/** Return the current phonon environment */
+		FEnvironment& GetPhononEnvironment();
+
+		/** Used to identify this plugin from other listeners - better if IAudioPluginListener.h has a GetPluginType() function instead */
+		static const FName PluginID;
+
+	private:
+		void InitializeEnvironment(FAudioDevice* AudioDevice, UWorld* ListenerWorld);
+
+		bool bEnvironmentInitialized;
+		FEnvironment Environment;
+		FPhononReverb* ReverbPtr;
+		FPhononOcclusion* OcclusionPtr;
+		AActor* SubLevelReference = nullptr;
+
+	private:
+		/* Helper function for checking whether the user is using Steam Audio for spatialization, reverb, and/or occlusion: */
+		static bool IsUsingSteamAudioPlugin(EAudioPlugin PluginType);
+		static UWorld* CurrentLevel;
+
+	};
+}
diff --git a/Engine/Plugins/Runtime/Steam/SteamAudio/Source/SteamAudio/Private/PhononProbeComponent.cpp b/Engine/Plugins/Runtime/Steam/SteamAudio/Source/SteamAudio/Private/PhononProbeComponent.cpp
new file mode 100644
--- /dev/null
+++ b/Engine/Plugins/Runtime/Steam/SteamAudio/Source/SteamAudio/Private/PhononProbeComponent.cpp
@@ -0,0 +1,5 @@
+//
+// Copyright (C) Valve Corporation. All rights reserved.
+//
+
+#include "PhononProbeComponent.h"
diff --git a/Engine/Plugins/Runtime/Steam/SteamAudio/Source/SteamAudio/Private/PhononProbeVolume.cpp b/Engine/Plugins/Runtime/Steam/SteamAudio/Source/SteamAudio/Private/PhononProbeVolume.cpp
new file mode 100644
--- /dev/null
+++ b/Engine/Plugins/Runtime/Steam/SteamAudio/Source/SteamAudio/Private/PhononProbeVolume.cpp
@@ -0,0 +1,313 @@
+//
+// Copyright (C) Valve Corporation. All rights reserved.
+//
+
+#include "PhononProbeVolume.h"
+
+#include "PhononProbeComponent.h"
+#include "PhononScene.h"
+#include "PhononCommon.h"
+
+#include "Components/PrimitiveComponent.h"
+#include "HAL/PlatformFileManager.h"
+#include "GenericPlatform/GenericPlatformFile.h"
+#include "Misc/Paths.h"
+#include "Engine/World.h"
+
+#if WITH_EDITOR
+
+#include "Editor.h"
+#include "LevelEditorViewport.h"
+
+#endif
+
+APhononProbeVolume::APhononProbeVolume(const FObjectInitializer& ObjectInitializer)
+	: Super(ObjectInitializer)
+	, PlacementStrategy(EPhononProbePlacementStrategy::UNIFORM_FLOOR)
+	, HorizontalSpacing(400.0f)
+	, HeightAboveFloor(150.0f)
+	, NumProbes(0)
+{
+	auto RootPrimitiveComponent = Cast<UPrimitiveComponent>(this->GetRootComponent());
+	RootPrimitiveComponent->BodyInstance.SetCollisionProfileName("NoCollision");
+	RootPrimitiveComponent->SetGenerateOverlapEvents(false);
+
+	FRotator DefaultRotation(0, 0, 0);
+	PhononProbeComponent = CreateDefaultSubobject<UPhononProbeComponent>(TEXT("PhononProbeComponent0"));
+	PhononProbeComponent->SetWorldLocation(this->GetActorLocation());
+	PhononProbeComponent->SetWorldRotation(DefaultRotation);
+	PhononProbeComponent->SetupAttachment(this->GetRootComponent());
+}
+
+#if WITH_EDITOR
+
+//==================================================================================================================================================
+// Probe placement
+//==================================================================================================================================================
+
+void APhononProbeVolume::PlaceProbes(IPLhandle PhononScene, IPLProbePlacementProgressCallback ProbePlacementCallback, 
+	TArray<IPLSphere>& ProbeSpheres)
+{
+	// TODO: delete probe batch data from disk if it exists
+
+	IPLhandle ProbeBox = nullptr;
+
+	// Compute box transform
+	float ProbeBoxTransformMatrix[16];
+	auto VolumeTransform = this->GetTransform();
+	VolumeTransform.MultiplyScale3D(FVector(200));
+	SteamAudio::GetMatrixForTransform(VolumeTransform, ProbeBoxTransformMatrix);
+
+	// Configure placement parameters
+	IPLProbePlacementParams ProbePlacementParameters;
+	ProbePlacementParameters.placement = PlacementStrategy == EPhononProbePlacementStrategy::CENTROID ? IPL_PLACEMENT_CENTROID : IPL_PLACEMENT_UNIFORMFLOOR;
+	ProbePlacementParameters.heightAboveFloor = HeightAboveFloor * SteamAudio::SCALEFACTOR;
+	ProbePlacementParameters.spacing = HorizontalSpacing * SteamAudio::SCALEFACTOR;
+	ProbePlacementParameters.maxOctreeDepth = 0;
+	ProbePlacementParameters.maxOctreeTriangles = 0;
+
+	// Create probe box, generate probes
+	iplCreateProbeBox(SteamAudio::GlobalContext, PhononScene, ProbeBoxTransformMatrix, ProbePlacementParameters, ProbePlacementCallback, &ProbeBox);
+
+	ProbeDataSize = SaveProbeBoxToDisk(ProbeBox);
+
+	// Get probe locations/radii
+	NumProbes = iplGetProbeSpheres(ProbeBox, nullptr);
+	ProbeSpheres.SetNumUninitialized(NumProbes);
+	iplGetProbeSpheres(ProbeBox, ProbeSpheres.GetData());
+
+	// Clean up
+	iplDestroyProbeBox(&ProbeBox);
+
+	MarkPackageDirty();
+}
+
+//==================================================================================================================================================
+// Probe box and probe batch serialization
+//==================================================================================================================================================
+
+int32 APhononProbeVolume::SaveProbeBatchToDisk(IPLhandle ProbeBatch)
+{
+	// Save probe batch data
+	TArray<uint8> ProbeBatchData;
+	int32 ProbeBatchDataSize = iplSaveProbeBatch(ProbeBatch, nullptr);
+	ProbeBatchData.SetNumUninitialized(ProbeBatchDataSize);
+	iplSaveProbeBatch(ProbeBatch, ProbeBatchData.GetData());
+
+	// Write data to disk
+	ProbeBatchFileName = ProbeBatchFileName.IsEmpty() ? GetProbeBatchFileName() : ProbeBatchFileName;
+	FString ProbeBatchPath = SteamAudio::RuntimePath + ProbeBatchFileName;
+	IPlatformFile& PlatformFile = FPlatformFileManager::Get().GetPlatformFile();
+	IFileHandle* ProbeBatchFileHandle = PlatformFile.OpenWrite(*ProbeBatchPath);
+	if (ProbeBatchFileHandle)
+	{
+		ProbeBatchFileHandle->Write(ProbeBatchData.GetData(), ProbeBatchDataSize);
+		delete ProbeBatchFileHandle;
+		return ProbeBatchDataSize;
+	}
+	else
+	{
+		UE_LOG(LogSteamAudio, Warning, TEXT("Unable to save probe batch to disk."));
+		return 0;
+	}
+}
+
+int32 APhononProbeVolume::SaveProbeBoxToDisk(IPLhandle ProbeBox)
+{
+	// Save probe box data
+	TArray<uint8> ProbeBoxData;
+	int32 ProbeBoxDataSize = iplSaveProbeBox(ProbeBox, nullptr);
+	ProbeBoxData.SetNumUninitialized(ProbeBoxDataSize);
+	iplSaveProbeBox(ProbeBox, ProbeBoxData.GetData());
+
+	// Write data to disk
+	ProbeBoxFileName = ProbeBoxFileName.IsEmpty() ? GetProbeBoxFileName() : ProbeBoxFileName;
+	FString ProbeBoxPath = SteamAudio::EditorOnlyPath + ProbeBoxFileName;
+	IPlatformFile& PlatformFile = FPlatformFileManager::Get().GetPlatformFile();
+	IFileHandle* ProbeBoxFileHandle = PlatformFile.OpenWrite(*ProbeBoxPath);
+	if (ProbeBoxFileHandle)
+	{
+		ProbeBoxFileHandle->Write(ProbeBoxData.GetData(), ProbeBoxDataSize);
+		delete ProbeBoxFileHandle;
+		return ProbeBoxDataSize;
+	}
+	else
+	{
+		UE_LOG(LogSteamAudio, Warning, TEXT("Unable to save probe box to disk."));
+		return 0;
+	}
+}
+
+//==================================================================================================================================================
+// Updating box/batch data.
+//==================================================================================================================================================
+
+void APhononProbeVolume::UpdateProbeData(IPLhandle ProbeBox)
+{
+	// Update probe box serialized data
+	SaveProbeBoxToDisk(ProbeBox);
+
+	// Update probe batch serialized data
+	IPLhandle ProbeBatch = nullptr;
+	iplCreateProbeBatch(SteamAudio::GlobalContext, &ProbeBatch);
+
+	NumProbes = iplGetProbeSpheres(ProbeBox, nullptr);
+	for (int32 i = 0; i < NumProbes; ++i)
+	{
+		iplAddProbeToBatch(ProbeBatch, ProbeBox, i);
+	}
+
+	iplFinalizeProbeBatch(ProbeBatch);
+	ProbeDataSize = SaveProbeBatchToDisk(ProbeBatch);
+	iplDestroyProbeBatch(&ProbeBatch);
+
+	MarkPackageDirty();
+}
+
+//==================================================================================================================================================
+// UI editability rules
+//==================================================================================================================================================
+
+bool APhononProbeVolume::CanEditChange(const FProperty* InProperty) const
+{
+	const bool ParentVal = Super::CanEditChange(InProperty);
+
+	if ((InProperty->GetFName() == GET_MEMBER_NAME_CHECKED(APhononProbeVolume, HorizontalSpacing)) ||
+		(InProperty->GetFName() == GET_MEMBER_NAME_CHECKED(APhononProbeVolume, HeightAboveFloor)))
+	{
+		return ParentVal && PlacementStrategy == EPhononProbePlacementStrategy::UNIFORM_FLOOR;
+	}
+
+	return ParentVal;
+}
+
+#endif
+
+//==================================================================================================================================================
+// Probe box and probe batch deserialization
+//==================================================================================================================================================
+
+bool APhononProbeVolume::LoadProbeBoxFromDisk(IPLhandle* ProbeBox) const
+{
+	TArray<uint8> ProbeBoxData;
+
+	if (ProbeBoxFileName.IsEmpty())
+	{
+		return false;
+	}
+
+	FString ProbeBoxPath = SteamAudio::EditorOnlyPath + ProbeBoxFileName;
+	IPlatformFile& PlatformFile = FPlatformFileManager::Get().GetPlatformFile();
+	IFileHandle* ProbeBoxFileHandle = PlatformFile.OpenRead(*ProbeBoxPath);
+	if (ProbeBoxFileHandle)
+	{
+		ProbeBoxData.SetNum(ProbeBoxFileHandle->Size());
+		ProbeBoxFileHandle->Read(ProbeBoxData.GetData(), ProbeBoxData.Num());
+		iplLoadProbeBox(SteamAudio::GlobalContext, ProbeBoxData.GetData(), ProbeBoxData.Num(), ProbeBox);
+		delete ProbeBoxFileHandle;
+	}
+	else
+	{
+		UE_LOG(LogSteamAudio, Warning, TEXT("Unable to load probe box from disk %s."), *ProbeBoxPath);
+		return false;
+	}
+
+	return true;
+}
+
+bool APhononProbeVolume::LoadProbeBatchFromDisk(IPLhandle* ProbeBatch) const
+{
+	TArray<uint8> ProbeBatchData;
+
+	if (ProbeBatchFileName.IsEmpty())
+	{
+		return false;
+	}
+
+	FString ProbeBatchPath = SteamAudio::RuntimePath + ProbeBatchFileName;
+	IPlatformFile& PlatformFile = FPlatformFileManager::Get().GetPlatformFile();
+	IFileHandle* ProbeBatchFileHandle = PlatformFile.OpenRead(*ProbeBatchPath);
+	if (ProbeBatchFileHandle)
+	{
+		ProbeBatchData.SetNum(ProbeBatchFileHandle->Size());
+		ProbeBatchFileHandle->Read(ProbeBatchData.GetData(), ProbeBatchData.Num());
+		iplLoadProbeBatch(SteamAudio::GlobalContext, ProbeBatchData.GetData(), ProbeBatchData.Num(), ProbeBatch);
+		delete ProbeBatchFileHandle;
+	}
+	else
+	{
+		UE_LOG(LogSteamAudio, Warning, TEXT("Unable to load probe batch from disk %s."), *ProbeBatchPath);
+		return false;
+	}
+
+	return true;
+}
+
+//==================================================================================================================================================
+// Data size queries
+//==================================================================================================================================================
+
+int32 APhononProbeVolume::GetProbeBoxDataSize() const
+{
+	int32 ProbeBoxDataSize = 0;
+	
+	if (!ProbeBoxFileName.IsEmpty())
+	{
+		FString ProbeBoxPath = SteamAudio::EditorOnlyPath + ProbeBoxFileName;
+		IPlatformFile& PlatformFile = FPlatformFileManager::Get().GetPlatformFile();
+		IFileHandle* ProbeBoxFileHandle = PlatformFile.OpenRead(*ProbeBoxPath);
+		if (ProbeBoxFileHandle)
+		{
+			ProbeBoxDataSize = ProbeBoxFileHandle->Size();
+			delete ProbeBoxFileHandle;
+		}
+	}
+	
+	return ProbeBoxDataSize;
+}
+
+int32 APhononProbeVolume::GetProbeBatchDataSize() const
+{
+	int32 ProbeBatchDataSize = 0;
+	
+	if (!ProbeBatchFileName.IsEmpty())
+	{
+		FString ProbeBatchPath = SteamAudio::RuntimePath + ProbeBatchFileName;
+		IPlatformFile& PlatformFile = FPlatformFileManager::Get().GetPlatformFile();
+		IFileHandle* ProbeBatchFileHandle = PlatformFile.OpenRead(*ProbeBatchPath);
+		if (ProbeBatchFileHandle)
+		{
+			ProbeBatchDataSize = ProbeBatchFileHandle->Size();
+			delete ProbeBatchFileHandle;
+		}
+	}
+
+	return ProbeBatchDataSize;
+}
+
+int32 APhononProbeVolume::GetDataSizeForSource(const FName& UniqueIdentifier) const
+{
+	int32 SourceDataSize = 0;
+	for (const auto& BakedSourceInfo : BakedDataInfo)
+	{
+		if (BakedSourceInfo.Name == UniqueIdentifier)
+		{
+			SourceDataSize += BakedSourceInfo.Size;
+		}
+	}
+	return SourceDataSize;
+}
+
+//==================================================================================================================================================
+// Filename generation
+//==================================================================================================================================================
+
+FString APhononProbeVolume::GetProbeBoxFileName() const
+{
+	return SteamAudio::StrippedMapName(GetWorld()->GetMapName()) + "_" + GetName() + ".probebox";
+}
+
+FString APhononProbeVolume::GetProbeBatchFileName() const
+{
+	return SteamAudio::StrippedMapName(GetWorld()->GetMapName()) + "_" + GetName() + ".probebatch";
+}
diff --git a/Engine/Plugins/Runtime/Steam/SteamAudio/Source/SteamAudio/Private/PhononReverb.cpp b/Engine/Plugins/Runtime/Steam/SteamAudio/Source/SteamAudio/Private/PhononReverb.cpp
new file mode 100644
--- /dev/null
+++ b/Engine/Plugins/Runtime/Steam/SteamAudio/Source/SteamAudio/Private/PhononReverb.cpp
@@ -0,0 +1,511 @@
+//
+// Copyright (C) Valve Corporation. All rights reserved.
+//
+
+#include "PhononReverb.h"
+#include "SteamAudioModule.h"
+#include "SteamAudioSettings.h"
+#include "PhononReverbSourceSettings.h"
+#include "Sound/SoundEffectSubmix.h"
+#include "Sound/SoundSubmix.h"
+#include "DSP/Dsp.h"
+#include "Misc/ScopeLock.h"
+#include "PhononPluginManager.h"
+#include "SteamAudioEnvironment.h"
+
+#include "AudioDevice.h"
+
+namespace SteamAudio
+{
+	//==============================================================================================================================================
+	// FPhononReverb
+	//==============================================================================================================================================
+
+	FPhononReverb::FPhononReverb()
+		: BinauralRenderer(nullptr)
+		, IndirectBinauralEffect(nullptr)
+		, IndirectPanningEffect(nullptr)
+		, ReverbConvolutionEffect(nullptr)
+		, AmbisonicsChannels(0)
+		, IndirectOutDeinterleaved(nullptr)
+		, bListenerInitialized(false)
+		, CachedSpatializationMethod(EIplSpatializationMethod::PANNING)
+		, Environment(nullptr)
+	{
+	}
+
+	FPhononReverb::~FPhononReverb()
+	{
+		for (auto& ReverbSource : ReverbSources)
+		{
+			if (ReverbSource.ConvolutionEffect)
+			{
+				iplDestroyConvolutionEffect(&ReverbSource.ConvolutionEffect);
+			}
+		}
+
+		if (ReverbConvolutionEffect)
+		{
+			iplDestroyConvolutionEffect(&ReverbConvolutionEffect);
+		}
+
+		if (IndirectBinauralEffect)
+		{
+			iplDestroyAmbisonicsBinauralEffect(&IndirectBinauralEffect);
+		}
+
+		if (IndirectPanningEffect)
+		{
+			iplDestroyAmbisonicsPanningEffect(&IndirectPanningEffect);
+		}
+
+		if (BinauralRenderer)
+		{
+			iplDestroyBinauralRenderer(&BinauralRenderer);
+		}
+
+		if (IndirectOutDeinterleaved)
+		{
+			for (int32 i = 0; i < AmbisonicsChannels; ++i)
+			{
+				delete[] IndirectOutDeinterleaved[i];
+			}
+			delete[] IndirectOutDeinterleaved;
+			IndirectOutDeinterleaved = nullptr;
+		}
+	}
+
+	// Just makes a copy of the init params - actual initialization needs to be deferred until the environment is created.
+	// This is because we do not know if we should fall back to Phonon settings from the TAN overrides until the compute
+	// device has been created.
+	void FPhononReverb::Initialize(const FAudioPluginInitializationParams InitializationParams)
+	{
+		bListenerInitialized = false;
+		AudioPluginInitializationParams = InitializationParams;
+
+		InputAudioFormat.channelLayout = IPL_CHANNELLAYOUT_MONO;
+		InputAudioFormat.channelLayoutType = IPL_CHANNELLAYOUTTYPE_SPEAKERS;
+		InputAudioFormat.channelOrder = IPL_CHANNELORDER_INTERLEAVED;
+		InputAudioFormat.numSpeakers = 1;
+		InputAudioFormat.speakerDirections = nullptr;
+		InputAudioFormat.ambisonicsOrder = -1;
+		InputAudioFormat.ambisonicsNormalization = IPL_AMBISONICSNORMALIZATION_N3D;
+		InputAudioFormat.ambisonicsOrdering = IPL_AMBISONICSORDERING_ACN;
+
+		ReverbSources.SetNum(AudioPluginInitializationParams.NumSources);
+		for (FReverbSource& ReverbSource : ReverbSources)
+		{
+			ReverbSource.InBuffer.format = InputAudioFormat;
+			ReverbSource.InBuffer.numSamples = AudioPluginInitializationParams.BufferLength;
+		}
+	}
+
+	void FPhononReverb::SetEnvironment(FEnvironment* InEnvironment)
+	{
+		if (!InEnvironment)
+		{
+			return;
+		}
+
+		Environment = InEnvironment;
+
+		const int32 IndirectImpulseResponseOrder = InEnvironment->GetSimulationSettings().ambisonicsOrder;
+		AmbisonicsChannels = (IndirectImpulseResponseOrder + 1) * (IndirectImpulseResponseOrder + 1);
+
+		ReverbInputAudioFormat.channelLayout = IPL_CHANNELLAYOUT_STEREO;
+		ReverbInputAudioFormat.channelLayoutType = IPL_CHANNELLAYOUTTYPE_SPEAKERS;
+		ReverbInputAudioFormat.channelOrder = IPL_CHANNELORDER_INTERLEAVED;
+		ReverbInputAudioFormat.numSpeakers = 2;
+		ReverbInputAudioFormat.speakerDirections = nullptr;
+		ReverbInputAudioFormat.ambisonicsOrder = -1;
+		ReverbInputAudioFormat.ambisonicsNormalization = IPL_AMBISONICSNORMALIZATION_N3D;
+		ReverbInputAudioFormat.ambisonicsOrdering = IPL_AMBISONICSORDERING_ACN;
+
+		IndirectOutputAudioFormat.channelLayout = IPL_CHANNELLAYOUT_MONO;
+		IndirectOutputAudioFormat.channelLayoutType = IPL_CHANNELLAYOUTTYPE_AMBISONICS;
+		IndirectOutputAudioFormat.channelOrder = IPL_CHANNELORDER_DEINTERLEAVED;
+		IndirectOutputAudioFormat.numSpeakers = (IndirectImpulseResponseOrder + 1) * (IndirectImpulseResponseOrder + 1);
+		IndirectOutputAudioFormat.speakerDirections = nullptr;
+		IndirectOutputAudioFormat.ambisonicsOrder = IndirectImpulseResponseOrder;
+		IndirectOutputAudioFormat.ambisonicsNormalization = IPL_AMBISONICSNORMALIZATION_N3D;
+		IndirectOutputAudioFormat.ambisonicsOrdering = IPL_AMBISONICSORDERING_ACN;
+
+		// Assume stereo output - if wrong, will be dynamically changed in the mixer processing
+		BinauralOutputAudioFormat.channelLayout = IPL_CHANNELLAYOUT_STEREO;
+		BinauralOutputAudioFormat.channelLayoutType = IPL_CHANNELLAYOUTTYPE_SPEAKERS;
+		BinauralOutputAudioFormat.channelOrder = IPL_CHANNELORDER_INTERLEAVED;
+		BinauralOutputAudioFormat.numSpeakers = 2;
+		BinauralOutputAudioFormat.speakerDirections = nullptr;
+		BinauralOutputAudioFormat.ambisonicsOrder = -1;
+		BinauralOutputAudioFormat.ambisonicsNormalization = IPL_AMBISONICSNORMALIZATION_N3D;
+		BinauralOutputAudioFormat.ambisonicsOrdering = IPL_AMBISONICSORDERING_ACN;
+
+		IPLHrtfParams HrtfParams;
+		HrtfParams.hrtfData = nullptr;
+		HrtfParams.sofaFileName = IPLstring("");
+		HrtfParams.type = IPL_HRTFDATABASETYPE_DEFAULT;
+
+		// The binaural renderer always uses Phonon convolution even if TAN is available.
+		IPLRenderingSettings BinauralRenderingSettings = Environment->GetRenderingSettings();
+		BinauralRenderingSettings.convolutionType = IPL_CONVOLUTIONTYPE_PHONON;
+
+		iplCreateBinauralRenderer(SteamAudio::GlobalContext, BinauralRenderingSettings, HrtfParams, &BinauralRenderer);
+
+		CachedSpatializationMethod = GetDefault<USteamAudioSettings>()->IndirectSpatializationMethod;
+
+		switch (CachedSpatializationMethod)
+		{
+		case EIplSpatializationMethod::HRTF:
+			iplCreateAmbisonicsBinauralEffect(BinauralRenderer, IndirectOutputAudioFormat, BinauralOutputAudioFormat, &IndirectBinauralEffect);
+			break;
+		case EIplSpatializationMethod::PANNING:
+			iplCreateAmbisonicsPanningEffect(BinauralRenderer, IndirectOutputAudioFormat, BinauralOutputAudioFormat, &IndirectPanningEffect);
+			break;
+		}
+
+		IndirectOutDeinterleaved = new IPLfloat32*[AmbisonicsChannels];
+		for (int32 i = 0; i < AmbisonicsChannels; ++i)
+		{
+			IndirectOutDeinterleaved[i] = new IPLfloat32[AudioPluginInitializationParams.BufferLength];
+		}
+
+		IndirectIntermediateBuffer.format = IndirectOutputAudioFormat;
+		IndirectIntermediateBuffer.numSamples = AudioPluginInitializationParams.BufferLength;
+		IndirectIntermediateBuffer.interleavedBuffer = nullptr;
+		IndirectIntermediateBuffer.deinterleavedBuffer = IndirectOutDeinterleaved;
+
+		DryBuffer.format = ReverbInputAudioFormat;
+		DryBuffer.numSamples = AudioPluginInitializationParams.BufferLength;
+		DryBuffer.interleavedBuffer = nullptr;
+		DryBuffer.deinterleavedBuffer = nullptr;
+
+		IndirectOutArray.SetNumZeroed(AudioPluginInitializationParams.BufferLength * BinauralOutputAudioFormat.numSpeakers);
+		IndirectOutBuffer.format = BinauralOutputAudioFormat;
+		IndirectOutBuffer.numSamples = AudioPluginInitializationParams.BufferLength;
+		IndirectOutBuffer.interleavedBuffer = IndirectOutArray.GetData();
+		IndirectOutBuffer.deinterleavedBuffer = nullptr;
+	}
+
+	void FPhononReverb::OnInitSource(const uint32 SourceId, const FName& AudioComponentUserId, const uint32 NumChannels, UReverbPluginSourceSettingsBase* InSettings)
+	{
+		if (!Environment || !Environment->GetEnvironmentalRenderer())
+		{
+			UE_LOG(LogSteamAudio, Error, TEXT("Unable to find environmental renderer for reverb. Reverb will not be applied. Make sure to export the scene."));
+			return;
+		}
+
+		FString SourceString = AudioComponentUserId.ToString().ToLower();
+		IPLBakedDataIdentifier SourceIdentifier;
+		SourceIdentifier.type = IPL_BAKEDDATATYPE_STATICSOURCE;
+		SourceIdentifier.identifier = Environment->GetBakedIdentifierMap().Get(SourceString);
+
+		UE_LOG(LogSteamAudio, Log, TEXT("Creating reverb effect for %s"), *SourceString);
+
+		UPhononReverbSourceSettings* Settings = static_cast<UPhononReverbSourceSettings*>(InSettings);
+		FReverbSource& ReverbSource = ReverbSources[SourceId];
+
+		ReverbSource.IndirectContribution = Settings->SourceReverbContribution;
+		ReverbSource.DipolePower = 0.0f;
+		ReverbSource.DipoleWeight = 0.0f;
+
+		InputAudioFormat.numSpeakers = NumChannels;
+		switch (NumChannels)
+		{
+			case 1: InputAudioFormat.channelLayout = IPL_CHANNELLAYOUT_MONO; break;
+			case 2: InputAudioFormat.channelLayout = IPL_CHANNELLAYOUT_STEREO; break;
+			case 4: InputAudioFormat.channelLayout = IPL_CHANNELLAYOUT_QUADRAPHONIC; break;
+			case 6: InputAudioFormat.channelLayout = IPL_CHANNELLAYOUT_FIVEPOINTONE; break;
+			case 8: InputAudioFormat.channelLayout = IPL_CHANNELLAYOUT_SEVENPOINTONE; break;
+		}
+
+		ReverbSource.InBuffer.format = InputAudioFormat;
+
+		switch (Settings->SourceReverbSimulationType)
+		{
+		case EIplSimulationType::BAKED:
+			iplCreateConvolutionEffect(Environment->GetEnvironmentalRenderer(), SourceIdentifier, IPL_SIMTYPE_BAKED, InputAudioFormat, IndirectOutputAudioFormat,
+				&ReverbSource.ConvolutionEffect);
+			break;
+		case EIplSimulationType::REALTIME:
+			iplCreateConvolutionEffect(Environment->GetEnvironmentalRenderer(), SourceIdentifier, IPL_SIMTYPE_REALTIME, InputAudioFormat, IndirectOutputAudioFormat,
+				&ReverbSource.ConvolutionEffect);
+			break;
+		case EIplSimulationType::DISABLED:
+		default:
+			break;
+		}
+	}
+
+	void FPhononReverb::OnReleaseSource(const uint32 SourceId)
+	{
+		UE_LOG(LogSteamAudio, Log, TEXT("Destroying reverb effect."));
+
+		check((int32)SourceId < ReverbSources.Num());
+
+		if (ReverbSources[SourceId].ConvolutionEffect)
+		{
+			iplDestroyConvolutionEffect(&ReverbSources[SourceId].ConvolutionEffect);
+		}
+	}
+
+	void FPhononReverb::ProcessSourceAudio(const FAudioPluginSourceInputData& InputData, FAudioPluginSourceOutputData& OutputData)
+	{
+		if (!Environment || !Environment->GetEnvironmentalRenderer() || !Environment->GetEnvironmentCriticalSectionHandle())
+		{
+			return;
+		}
+
+		FScopeLock EnvironmentLock(Environment->GetEnvironmentCriticalSectionHandle());
+
+		FReverbSource& ReverbSource = ReverbSources[InputData.SourceId];
+
+		IPLSource SourceData;
+		SourceData.position = SteamAudio::UnrealToPhononIPLVector3(InputData.SpatializationParams->EmitterWorldPosition);
+		SourceData.ahead = SteamAudio::UnrealToPhononIPLVector3(InputData.SpatializationParams->EmitterWorldRotation * FVector::ForwardVector);
+		SourceData.right = SteamAudio::UnrealToPhononIPLVector3(InputData.SpatializationParams->EmitterWorldRotation * FVector::RightVector);
+		SourceData.up = SteamAudio::UnrealToPhononIPLVector3(InputData.SpatializationParams->EmitterWorldRotation * FVector::UpVector);
+		SourceData.directivity = IPLDirectivity{ ReverbSource.DipolePower, ReverbSource.DipoleWeight, nullptr, nullptr };
+
+		if (ReverbSource.ConvolutionEffect)
+		{
+			ReverbSource.IndirectInArray.SetNumUninitialized(InputData.AudioBuffer->Num());
+			for (int32 i = 0; i < InputData.AudioBuffer->Num(); ++i)
+			{
+				ReverbSource.IndirectInArray[i] = (*InputData.AudioBuffer)[i] * ReverbSource.IndirectContribution;
+			}
+			ReverbSource.InBuffer.interleavedBuffer = ReverbSource.IndirectInArray.GetData();
+
+			iplSetDryAudioForConvolutionEffect(ReverbSource.ConvolutionEffect, SourceData, ReverbSource.InBuffer);
+		}
+	}
+
+	void FPhononReverb::ProcessMixedAudio(const FSoundEffectSubmixInputData& InData, FSoundEffectSubmixOutputData& OutData)
+	{
+		if (!Environment || !Environment->GetEnvironmentalRenderer() || !Environment->GetEnvironmentCriticalSectionHandle() || !bListenerInitialized)
+		{
+			return;
+		}
+
+		FScopeLock EnvironmentLock(Environment->GetEnvironmentCriticalSectionHandle());
+
+		if (IndirectOutBuffer.format.numSpeakers != OutData.NumChannels)
+		{
+			if (IndirectBinauralEffect)
+			{
+				iplDestroyAmbisonicsBinauralEffect(&IndirectBinauralEffect);
+			}
+
+			if (IndirectPanningEffect)
+			{
+				iplDestroyAmbisonicsPanningEffect(&IndirectPanningEffect);
+			}
+
+			IndirectOutBuffer.format.numSpeakers = OutData.NumChannels;
+			switch (OutData.NumChannels)
+			{
+				case 1: IndirectOutBuffer.format.channelLayout = IPL_CHANNELLAYOUT_MONO; break;
+				case 2: IndirectOutBuffer.format.channelLayout = IPL_CHANNELLAYOUT_STEREO; break;
+				case 4: IndirectOutBuffer.format.channelLayout = IPL_CHANNELLAYOUT_QUADRAPHONIC; break;
+				case 6: IndirectOutBuffer.format.channelLayout = IPL_CHANNELLAYOUT_FIVEPOINTONE; break;
+				case 8: IndirectOutBuffer.format.channelLayout = IPL_CHANNELLAYOUT_SEVENPOINTONE; break;
+			}
+
+			IndirectOutArray.SetNumZeroed(OutData.AudioBuffer->Num());
+			IndirectOutBuffer.interleavedBuffer = IndirectOutArray.GetData();
+
+			switch (CachedSpatializationMethod)
+			{
+			case EIplSpatializationMethod::HRTF:
+				iplCreateAmbisonicsBinauralEffect(BinauralRenderer, IndirectOutputAudioFormat, IndirectOutBuffer.format, &IndirectBinauralEffect);
+				break;
+			case EIplSpatializationMethod::PANNING:
+				iplCreateAmbisonicsPanningEffect(BinauralRenderer, IndirectOutputAudioFormat, IndirectOutBuffer.format, &IndirectPanningEffect);
+				break;
+			}
+		}
+
+		if (ReverbConvolutionEffect)
+		{
+			float ListenerReverbContribution = GetDefault<USteamAudioSettings>()->ListenerReverbContribution;
+			ReverbIndirectInArray.SetNumUninitialized(InData.AudioBuffer->Num());
+			for (int32 i = 0; i < InData.AudioBuffer->Num(); ++i)
+			{
+				ReverbIndirectInArray[i] = (*InData.AudioBuffer)[i] * ListenerReverbContribution;
+			}
+
+			IPLSource ReverbSource = {};
+			ReverbSource.position = ListenerPosition;
+			ReverbSource.ahead = ListenerForward;
+			ReverbSource.up = ListenerUp;
+			ReverbSource.directivity = IPLDirectivity{ 0.0f, 0.0f, nullptr, nullptr };
+
+			DryBuffer.interleavedBuffer = ReverbIndirectInArray.GetData();
+			iplSetDryAudioForConvolutionEffect(ReverbConvolutionEffect, ReverbSource, DryBuffer);
+		}
+
+		iplGetMixedEnvironmentalAudio(Environment->GetEnvironmentalRenderer(), ListenerPosition, ListenerForward, ListenerUp, IndirectIntermediateBuffer);
+
+		switch (CachedSpatializationMethod)
+		{
+		case EIplSpatializationMethod::HRTF:
+			iplApplyAmbisonicsBinauralEffect(IndirectBinauralEffect, BinauralRenderer, IndirectIntermediateBuffer, IndirectOutBuffer);
+			break;
+		case EIplSpatializationMethod::PANNING:
+			iplApplyAmbisonicsPanningEffect(IndirectPanningEffect, BinauralRenderer, IndirectIntermediateBuffer, IndirectOutBuffer);
+			break;
+		}
+
+		FMemory::Memcpy(OutData.AudioBuffer->GetData(), IndirectOutArray.GetData(), sizeof(float) * IndirectOutArray.Num());
+	}
+
+	void FPhononReverb::CreateReverbEffect()
+	{
+		check(Environment && Environment->GetEnvironmentalRenderer() && Environment->GetEnvironmentCriticalSectionHandle());
+
+		IPLBakedDataIdentifier ReverbIdentifier;
+		ReverbIdentifier.type = IPL_BAKEDDATATYPE_REVERB;
+		ReverbIdentifier.identifier = 0;
+
+		switch (GetDefault<USteamAudioSettings>()->ListenerReverbSimulationType)
+		{
+		case EIplSimulationType::BAKED:
+			iplCreateConvolutionEffect(Environment->GetEnvironmentalRenderer(), ReverbIdentifier, IPL_SIMTYPE_BAKED, ReverbInputAudioFormat, IndirectOutputAudioFormat,
+				&ReverbConvolutionEffect);
+			break;
+		case EIplSimulationType::REALTIME:
+			iplCreateConvolutionEffect(Environment->GetEnvironmentalRenderer(), ReverbIdentifier, IPL_SIMTYPE_REALTIME, ReverbInputAudioFormat, IndirectOutputAudioFormat,
+				&ReverbConvolutionEffect);
+			break;
+		case EIplSimulationType::DISABLED:
+		default:
+			break;
+		}
+	}
+
+	void FPhononReverb::UpdateListener(const FVector& Position, const FVector& Forward, const FVector& Up, const FVector& Right)
+	{
+		ListenerPosition = SteamAudio::UnrealToPhononIPLVector3(Position);
+		ListenerForward = SteamAudio::UnrealToPhononIPLVector3(Forward, false);
+		ListenerUp = SteamAudio::UnrealToPhononIPLVector3(Up, false);
+		ListenerRight = SteamAudio::UnrealToPhononIPLVector3(Right, false);
+		bListenerInitialized = true;
+	}
+
+	FSoundEffectSubmixPtr FPhononReverb::GetEffectSubmix()
+	{
+		if (!SubmixEffect.IsValid())
+		{
+			USubmixEffectReverbPluginPreset* ReverbPluginPreset = nullptr;
+			USoundSubmix* Submix = GetSubmix();
+			if (Submix->SubmixEffectChain.Num() > 0)
+			{
+				if (USubmixEffectReverbPluginPreset* CurrentPreset = Cast<USubmixEffectReverbPluginPreset>(Submix->SubmixEffectChain[0]))
+				{
+					ReverbPluginPreset = CurrentPreset;
+				}
+			}
+
+			if (!ReverbPluginPreset)
+			{
+				ReverbPluginPreset = NewObject<USubmixEffectReverbPluginPreset>(Submix, TEXT("Reverb Plugin Effect Preset"));
+			}
+
+			SubmixEffect = USoundEffectPreset::CreateInstance<FSoundEffectSubmixInitData, FSoundEffectSubmix>(FSoundEffectSubmixInitData(), *ReverbPluginPreset);
+			StaticCastSharedPtr<FSubmixEffectReverbPlugin, FSoundEffectSubmix, ESPMode::ThreadSafe>(SubmixEffect)->SetPhononReverbPlugin(this);
+			SubmixEffect->SetEnabled(true);
+		}
+
+		return SubmixEffect;
+	}
+
+	USoundSubmix* FPhononReverb::GetSubmix()
+	{
+		const USteamAudioSettings* Settings = GetDefault<USteamAudioSettings>();
+		check(Settings);
+		
+		if (!ReverbSubmix.IsValid())
+		{
+			ReverbSubmix = Cast<USoundSubmix>(Settings->OutputSubmix.TryLoad());
+		}
+
+		if (!ReverbSubmix.IsValid())
+		{
+			static const FString DefaultSubmixName = TEXT("Phonon Reverb Submix");
+			UE_LOG(LogSteamAudio, Error, TEXT("Failed to load Phonon Reverb Submix from object path '%s' in PhononSettings. Creating '%s' as stub."),
+				*Settings->OutputSubmix.GetAssetPathString(),
+				*DefaultSubmixName);
+
+			ReverbSubmix = NewObject<USoundSubmix>(USoundSubmix::StaticClass(), *DefaultSubmixName);
+			ReverbSubmix->bMuteWhenBackgrounded = true;
+		}
+		ReverbSubmix->bAutoDisable = false;
+
+		bool bFoundPreset = false;
+		for (USoundEffectSubmixPreset* Preset : ReverbSubmix->SubmixEffectChain)
+		{
+			if (USubmixEffectReverbPluginPreset* PluginPreset = Cast<USubmixEffectReverbPluginPreset>(Preset))
+			{
+				bFoundPreset = true;
+				break;
+			}
+		}
+
+		if (!bFoundPreset)
+		{
+			static const FString DefaultPresetName = TEXT("PhononReverbDefault_0");
+			UE_LOG(LogSteamAudio, Error, TEXT("Failed to find Phonon USubmixEffectReverbPluginPreset on default Phonon Reverb Submix. Creating stub '%s'."),
+				*Settings->OutputSubmix.GetAssetPathString(),
+				*DefaultPresetName);
+			ReverbSubmix->SubmixEffectChain.Add(NewObject<USubmixEffectReverbPluginPreset>(USubmixEffectReverbPluginPreset::StaticClass(), *DefaultPresetName));
+		}
+
+		return ReverbSubmix.Get();
+	}
+
+	//==============================================================================================================================================
+	// FReverbSource
+	//==============================================================================================================================================
+
+	FReverbSource::FReverbSource()
+		: ConvolutionEffect(nullptr)
+		, IndirectContribution(1.0f)
+		, DipolePower(0.0f)
+		, DipoleWeight(0.0f)
+	{
+	}
+
+}
+
+//==================================================================================================================================================
+// FSubmixEffectReverbPlugin
+//==================================================================================================================================================
+
+FSubmixEffectReverbPlugin::FSubmixEffectReverbPlugin()
+	: PhononReverbPlugin(nullptr)
+{}
+
+void FSubmixEffectReverbPlugin::Init(const FSoundEffectSubmixInitData& InInitData)
+{
+}
+
+void FSubmixEffectReverbPlugin::OnPresetChanged()
+{
+}
+
+uint32 FSubmixEffectReverbPlugin::GetDesiredInputChannelCountOverride() const
+{
+	return 2;
+}
+
+void FSubmixEffectReverbPlugin::OnProcessAudio(const FSoundEffectSubmixInputData& InData, FSoundEffectSubmixOutputData& OutData)
+{
+	if (PhononReverbPlugin)
+	{
+		PhononReverbPlugin->ProcessMixedAudio(InData, OutData);
+	}
+}
+
+void FSubmixEffectReverbPlugin::SetPhononReverbPlugin(SteamAudio::FPhononReverb* InPhononReverbPlugin)
+{
+	PhononReverbPlugin = InPhononReverbPlugin;
+}
diff --git a/Engine/Plugins/Runtime/Steam/SteamAudio/Source/SteamAudio/Private/PhononReverb.h b/Engine/Plugins/Runtime/Steam/SteamAudio/Source/SteamAudio/Private/PhononReverb.h
new file mode 100644
--- /dev/null
+++ b/Engine/Plugins/Runtime/Steam/SteamAudio/Source/SteamAudio/Private/PhononReverb.h
@@ -0,0 +1,133 @@
+//
+// Copyright (C) Valve Corporation. All rights reserved.
+//
+
+#pragma once
+
+#include "IAudioExtensionPlugin.h"
+#include "Sound/SoundEffectSubmix.h"
+#include "Sound/SoundEffectPreset.h"
+#include "UObject/WeakObjectPtrTemplates.h"
+
+#include "PhononCommon.h"
+#include "phonon.h"
+#include "PhononReverb.generated.h"
+
+
+// Forward Declaration
+class FSubmixEffectReverbPlugin;
+
+
+namespace SteamAudio
+{
+	class FEnvironment;
+
+	struct FReverbSource
+	{
+		FReverbSource();
+
+		IPLhandle ConvolutionEffect;
+		float IndirectContribution;
+		float DipolePower;
+		float DipoleWeight;
+		IPLAudioBuffer InBuffer;
+		TArray<float> IndirectInArray;
+	};
+
+	/************************************************************************/
+	/* FPhononReverb. Reverb plugin for Steam Audio.                        */
+	/************************************************************************/
+	class FPhononReverb : public IAudioReverb
+	{
+	public:
+		FPhononReverb();
+		~FPhononReverb();
+
+		virtual void Initialize(const FAudioPluginInitializationParams InitializationParams) override;
+		virtual void OnInitSource(const uint32 SourceId, const FName& AudioComponentUserId, const uint32 NumChannels, UReverbPluginSourceSettingsBase* InSettings) override;
+		virtual void OnReleaseSource(const uint32 SourceId) override;
+		virtual FSoundEffectSubmixPtr GetEffectSubmix() override;
+		virtual USoundSubmix* GetSubmix() override;
+		virtual void ProcessSourceAudio(const FAudioPluginSourceInputData& InputData, FAudioPluginSourceOutputData& OutputData) override;
+		
+		void ProcessMixedAudio(const FSoundEffectSubmixInputData& InData, FSoundEffectSubmixOutputData& OutData);
+		void SetEnvironment(FEnvironment* InEnvironment);
+		void CreateReverbEffect();
+		void UpdateListener(const FVector& Position, const FVector& Forward, const FVector& Up, const FVector& Right);
+
+	private:
+		FSoundEffectSubmixPtr SubmixEffect;
+		TWeakObjectPtr<USoundSubmix> ReverbSubmix;
+
+		IPLhandle BinauralRenderer;
+		IPLhandle IndirectBinauralEffect;
+		IPLhandle IndirectPanningEffect;
+		IPLhandle ReverbConvolutionEffect;
+		IPLAudioBuffer DryBuffer;
+
+		IPLAudioBuffer IndirectOutBuffer;
+		int32 AmbisonicsChannels;
+		IPLfloat32** IndirectOutDeinterleaved;
+		TArray<float> IndirectOutArray;
+
+		IPLAudioBuffer IndirectIntermediateBuffer;
+		TArray<float> IndirectIntermediateArray;
+
+		IPLAudioFormat InputAudioFormat;
+		IPLAudioFormat ReverbInputAudioFormat;
+		IPLAudioFormat IndirectOutputAudioFormat;
+		IPLAudioFormat BinauralOutputAudioFormat;
+
+		FCriticalSection ListenerCriticalSection;
+		IPLVector3 ListenerPosition;
+		IPLVector3 ListenerForward;
+		IPLVector3 ListenerRight;
+		IPLVector3 ListenerUp;
+		bool bListenerInitialized;
+
+		EIplSpatializationMethod CachedSpatializationMethod;
+
+		TArray<FReverbSource> ReverbSources;
+
+		TArray<float> ReverbIndirectInArray;
+
+		TAudioPluginListenerPtr PluginManagerPtr;
+
+		FAudioPluginInitializationParams AudioPluginInitializationParams;
+		FEnvironment* Environment;
+	};
+}
+
+class FSubmixEffectReverbPlugin : public FSoundEffectSubmix
+{
+public:
+	FSubmixEffectReverbPlugin();
+
+	virtual void Init(const FSoundEffectSubmixInitData& InSampleRate) override;
+	virtual uint32 GetDesiredInputChannelCountOverride() const override;
+	virtual void OnProcessAudio(const FSoundEffectSubmixInputData& InData, FSoundEffectSubmixOutputData& OutData) override;
+	virtual void OnPresetChanged() override;
+
+	void SetPhononReverbPlugin(SteamAudio::FPhononReverb* PhononReverbPlugin);
+
+private:
+	SteamAudio::FPhononReverb* PhononReverbPlugin;
+};
+
+USTRUCT()
+struct FSubmixEffectReverbPluginSettings
+{
+	GENERATED_USTRUCT_BODY()
+};
+
+UCLASS()
+class USubmixEffectReverbPluginPreset : public USoundEffectSubmixPreset
+{
+	GENERATED_BODY()
+
+public:
+	EFFECT_PRESET_METHODS(SubmixEffectReverbPlugin)
+
+	UPROPERTY(EditAnywhere, Category = SubmixEffectPreset)
+	FSubmixEffectReverbPluginSettings Settings;
+};
diff --git a/Engine/Plugins/Runtime/Steam/SteamAudio/Source/SteamAudio/Private/PhononReverbSourceSettings.cpp b/Engine/Plugins/Runtime/Steam/SteamAudio/Source/SteamAudio/Private/PhononReverbSourceSettings.cpp
new file mode 100644
--- /dev/null
+++ b/Engine/Plugins/Runtime/Steam/SteamAudio/Source/SteamAudio/Private/PhononReverbSourceSettings.cpp
@@ -0,0 +1,10 @@
+//
+// Copyright (C) Valve Corporation. All rights reserved.
+//
+
+#include "PhononReverbSourceSettings.h"
+
+UPhononReverbSourceSettings::UPhononReverbSourceSettings()
+	: SourceReverbSimulationType(EIplSimulationType::REALTIME)
+	, SourceReverbContribution(1.0f)
+{}
\ No newline at end of file
