[Backout] - CL25430596

diff --git a/Engine/Plugins/Runtime/Steam/SteamAudio/Source/SteamAudio/Private/PhononScene.cpp b/Engine/Plugins/Runtime/Steam/SteamAudio/Source/SteamAudio/Private/PhononScene.cpp
new file mode 100644
--- /dev/null
+++ b/Engine/Plugins/Runtime/Steam/SteamAudio/Source/SteamAudio/Private/PhononScene.cpp
@@ -0,0 +1,1030 @@
+//
+// Copyright (C) Valve Corporation. All rights reserved.
+//
+
+#include "PhononScene.h"
+
+#include "Model.h"
+#include "PhononCommon.h"
+#include "PhononGeometryComponent.h"
+#include "PhononMaterialComponent.h"
+
+#include "EngineUtils.h"
+#include "HAL/FileManagerGeneric.h"
+#include "Engine/StaticMeshActor.h"
+#include "Engine/StaticMesh.h"
+#include "StaticMeshResources.h"
+#include "SteamAudioSettings.h"
+#include "HAL/PlatformFileManager.h"
+#include "GenericPlatform/GenericPlatformFile.h"
+#include "Misc/Paths.h"
+#include "Async/Async.h"
+
+#if WITH_EDITOR
+
+#include "ISettingsModule.h"
+#include "LevelEditor.h"
+#include "LevelEditorActions.h"
+#include "Landscape.h"
+#include "LandscapeComponent.h"
+#include "LandscapeInfo.h"
+#include "LandscapeDataAccess.h"
+
+#endif // WITH_EDITOR
+
+/*
+
+ The scene export functions set up the following material index layout on the Phonon backend:
+
+ <Presets>
+ Default static mesh material
+ Default BSP material
+ Default landscape material
+ <Custom static mesh materials>
+
+ Note that it results in the CUSTOM preset being unused, but the code is simpler this way.
+
+*/
+
+namespace SteamAudio
+{
+	bool LoadSceneFromDisk(UWorld* World, IPLhandle ComputeDevice, const IPLSimulationSettings& SimulationSettings, IPLhandle* PhononScene,
+		FPhononSceneInfo& PhononSceneInfo, FString* CurrentLevelName)
+	{
+		FString MapName;
+		
+		if (!CurrentLevelName || CurrentLevelName->IsEmpty())
+		{
+			MapName = World->GetMapName();
+		}
+		else
+		{
+			MapName = StrippedMapName(*CurrentLevelName);
+		}
+		
+		FString SceneFileName = RuntimePath + MapName + ".phononscene";
+		FString SceneInfoFileName = EditorOnlyPath + MapName + ".phononsceneinfo";
+		TArray<uint8> SceneData;
+
+		IPlatformFile& PlatformFile = FPlatformFileManager::Get().GetPlatformFile();
+		IFileHandle* SceneFileHandle = PlatformFile.OpenRead(*SceneFileName);
+		if (SceneFileHandle)
+		{
+			SceneData.SetNum(SceneFileHandle->Size());
+			SceneFileHandle->Read(SceneData.GetData(), SceneFileHandle->Size());
+			delete SceneFileHandle;
+		}
+		else
+		{
+			UE_LOG(LogSteamAudio, Warning, TEXT("Unable to load Phonon scene: error reading file. Be sure to export the scene."));
+			return false;
+		}
+
+		IPLerror Error = iplLoadScene(GlobalContext, SimulationSettings, SceneData.GetData(), SceneData.Num(), ComputeDevice, nullptr, PhononScene);
+		if (Error != IPL_STATUS_SUCCESS)
+		{
+			UE_LOG(LogSteamAudio, Warning, TEXT("Unable to load Phonon scene: error [%i] loading scene file [%s] with data size [%i]. \n "), (int32) Error, *SceneFileName, SceneData.Num());
+			return false;
+		}
+
+		LoadSceneInfoFromDisk(World, PhononSceneInfo);
+
+		return true;
+	}
+
+	bool LoadSceneInfoFromDisk(UWorld* World, FPhononSceneInfo& PhononSceneInfo)
+	{
+		FString MapName = StrippedMapName(World->GetMapName());
+		FString SceneInfoFileName = EditorOnlyPath + MapName + ".phononsceneinfo";
+
+		IPlatformFile& PlatformFile = FPlatformFileManager::Get().GetPlatformFile();
+		IFileHandle* SceneInfoFileHandle = PlatformFile.OpenRead(*SceneInfoFileName);
+		if (SceneInfoFileHandle)
+		{
+			SceneInfoFileHandle->Read((uint8*)&PhononSceneInfo.NumTriangles, 4);
+			SceneInfoFileHandle->Read((uint8*)&PhononSceneInfo.NumDynTriangles, 4);
+			SceneInfoFileHandle->Read((uint8*)&PhononSceneInfo.DataSize, 4);
+			SceneInfoFileHandle->Read((uint8*)&PhononSceneInfo.DynDataSize, 4);
+			delete SceneInfoFileHandle;
+		}
+		else
+		{
+			UE_LOG(LogSteamAudio, Warning, TEXT("Unable to load Phonon scene info: error loading info file."));
+			return false;
+		}
+
+		return true;
+	}
+
+	bool LoadDynamicSceneFromDisk(UWorld* World, IPLhandle ComputeDevice, const IPLSimulationSettings& SimulationSettings, IPLhandle* PhononScene,
+		FPhononSceneInfo& PhononSceneInfo, FString DynamicSceneID)
+	{
+		// Check for supported raytracer
+		if (GetDefault<USteamAudioSettings>()->RayTracer != EIplRayTracerType::EMBREE)
+		{
+			return false;
+		}
+
+		FString SceneFileName = DynamicRuntimePath + DynamicSceneID + ".phononscene";
+		FString SceneInfoFileName = DynamicEditorOnlyPath + DynamicSceneID + ".phononinfo";
+		TArray<uint8> SceneData;
+
+		IPlatformFile& PlatformFile = FPlatformFileManager::Get().GetPlatformFile();
+
+		IFileHandle* SceneFileHandle = PlatformFile.OpenRead(*SceneFileName);
+		if (SceneFileHandle)
+		{
+			SceneData.SetNum(SceneFileHandle->Size());
+			SceneFileHandle->Read(SceneData.GetData(), SceneFileHandle->Size());
+			delete SceneFileHandle;
+		}
+		else
+		{
+			UE_LOG(LogSteamAudio, Warning, TEXT("Unable to load Dynamic Phonon scene: error reading file. Be sure to export the scene."));
+			return false;
+		}
+
+		IPLerror Error = iplLoadScene(GlobalContext, SimulationSettings, SceneData.GetData(), SceneData.Num(), ComputeDevice, nullptr, PhononScene);
+		if (Error != IPL_STATUS_SUCCESS)
+		{
+			UE_LOG(LogSteamAudio, Warning, TEXT("Unable to load Dynamic Phonon scene: error [%i] loading scene file [%s] with data size [%i]. \n "), (int32)Error, *SceneFileName, SceneData.Num());
+			return false;
+		}
+
+		LoadDynamicScenesInfoFromDisk(World, PhononSceneInfo, DynamicSceneID);
+
+		return true;
+	}
+
+	bool LoadDynamicScenesInfoFromDisk(UWorld* World, FPhononSceneInfo& PhononSceneInfo, FString DynamicSceneID)
+	{
+		FString SceneInfoFileName = DynamicEditorOnlyPath + DynamicSceneID + ".phononinfo";
+
+		IPlatformFile& PlatformFile = FPlatformFileManager::Get().GetPlatformFile();
+		IFileHandle* SceneInfoFileHandle = PlatformFile.OpenRead(*SceneInfoFileName);
+		if (SceneInfoFileHandle)
+		{
+			SceneInfoFileHandle->Read((uint8*)&PhononSceneInfo.NumTriangles, 4);
+			SceneInfoFileHandle->Seek(4);
+			SceneInfoFileHandle->Read((uint8*)&PhononSceneInfo.DataSize, 4);
+			delete SceneInfoFileHandle;
+		}
+		else
+		{
+			UE_LOG(LogSteamAudio, Warning, TEXT("Unable to load Dynamic Phonon scene info: error loading info file."));
+			return false;
+		}
+
+		return true;
+	}
+
+#if WITH_EDITOR
+
+	static void LoadBSPGeometry(UWorld* World, IPLhandle PhononScene, TArray<IPLVector3>& IplVertices, TArray<IPLTriangle>& IplTriangles,
+		TArray<IPLint32>& IplMaterialIndices);
+	static void LoadStaticMeshActors(UWorld* World, IPLhandle PhononScene, TArray<IPLVector3>& IplVertices, TArray<IPLTriangle>& IplTriangles,
+		TArray<IPLint32>& IplMaterialIndices);
+	static void LoadLandscapeActors(UWorld* World, IPLhandle PhononScene, TArray<IPLVector3>& IplVertices, TArray<IPLTriangle>& IplTriangles,
+		TArray<IPLint32>& IplMaterialIndices);
+	static void RegisterStaticMesh(IPLhandle PhononScene, TArray<IPLVector3>& IplVertices, TArray<IPLTriangle>& IplTriangles,
+		TArray<IPLint32>& IplMaterialIndices, IPLhandle* PhononStaticMesh);
+	static void RegisterStaticMeshes(IPLhandle PhononScene, TArray<IPLVector3>& IplVertices, TArray<IPLTriangle>& IplTriangles,
+		TArray<IPLint32>& IplMaterialIndices, TArray<IPLhandle>* PhononStaticMeshes);
+	static void GetPhononMaterials(UWorld* World, TArray<IPLMaterial>& IplMaterials);
+	static UPhononMaterialComponent* GetPhononMaterialComponent(AActor* Actor);
+	static uint32 GetMeshVerts(UStaticMeshComponent* StaticMeshComponent, TArray<IPLVector3>& VertexArray, bool bGetRelativeVertexPos = true);
+	static void CalculateVertsTrisMats(AActor* MeshActor, UStaticMeshComponent* MeshComponent, TArray<IPLVector3>& IplVertices, TArray<IPLTriangle>& IplTriangles,
+		TArray<IPLint32>& IplMaterialIndices);
+
+	//==============================================================================================================================================
+	// High level scene export
+	//==============================================================================================================================================
+
+	/**
+	 * Loads scene geometry, providing handles to the Phonon scene object and Phonon static meshes.
+	 */
+	bool CreateScene(UWorld* World, IPLhandle* PhononScene, IPLhandle* PhononStaticMesh, uint32& NumSceneTriangles, uint32& NumDynSceneTriangles, uint32& DynDataSize, bool CreateDynObjFiles /* =false*/ )
+	{
+		check(World);
+		check(PhononScene);
+		check(PhononStaticMesh);
+
+		UE_LOG(LogSteamAudio, Log, TEXT("Loading Phonon scene."));
+
+		TPromise<bool> Result;
+
+		AsyncTask(ENamedThreads::GameThread, [World, PhononScene, PhononStaticMesh, &NumSceneTriangles, &NumDynSceneTriangles, &DynDataSize, CreateDynObjFiles, &Result]()
+		{
+			// Setup simulation settings - we really only care about scene type here, but Beta 17 API throws if ambisonics order > max order,
+			// so just zero the struct for now.
+			IPLSimulationSettings SimulationSettings;
+			FMemory::Memset(SimulationSettings, 0);
+			SimulationSettings.sceneType = IPLSceneType::IPL_SCENETYPE_PHONON;
+
+			// Setup IPL Material defaults - custom will be unused
+			TArray<IPLMaterial> SceneMaterials;
+			MaterialPresets.GenerateValueArray(SceneMaterials);
+
+			// Add default materials as fixed slots (3)
+			SceneMaterials.Add(GetDefault<USteamAudioSettings>()->GetDefaultStaticMeshMaterial());
+			SceneMaterials.Add(GetDefault<USteamAudioSettings>()->GetDefaultBSPMaterial());
+			SceneMaterials.Add(GetDefault<USteamAudioSettings>()->GetDefaultLandscapeMaterial());
+			
+			// Get IPLMaterials in scene
+			GetPhononMaterials(World, SceneMaterials);
+			UE_LOG(LogSteamAudio, Log, TEXT("Loading Phonon scene found [%i] materials."), SceneMaterials.Num());
+			
+			// Ensure that the SteamAudio paths exist, create them if they don't
+			CheckSceneDirectory(RuntimePath);
+			CheckSceneDirectory(EditorOnlyPath);
+
+			// Create dynamic scene
+			CreateDynamicScene(World, SimulationSettings, SceneMaterials, NumDynSceneTriangles, DynDataSize, CreateDynObjFiles);
+		
+			// Create static scene
+			IPLerror IplResult = iplCreateScene(GlobalContext, nullptr, SimulationSettings, SceneMaterials.Num(),
+				SceneMaterials.GetData(), nullptr, nullptr, nullptr, nullptr, nullptr, PhononScene);
+
+			if (IplResult != IPL_STATUS_SUCCESS)
+			{
+				UE_LOG(LogSteamAudio, Warning, TEXT("Error creating Phonon scene."));
+				Result.SetValue(false);
+				return;
+			}
+
+			TArray<IPLVector3> IplVertices;
+			TArray<IPLTriangle> IplTriangles;
+			TArray<IPLint32> IplMaterialIndices;
+
+			LoadStaticMeshActors(World, *PhononScene, IplVertices, IplTriangles, IplMaterialIndices);
+
+			if (GetDefault<USteamAudioSettings>()->ExportLandscapeGeometry)
+			{
+				LoadLandscapeActors(World, *PhononScene, IplVertices, IplTriangles, IplMaterialIndices);
+			}
+
+			if (GetDefault<USteamAudioSettings>()->ExportBSPGeometry)
+			{
+				LoadBSPGeometry(World, *PhononScene, IplVertices, IplTriangles, IplMaterialIndices);
+			}
+			
+			RegisterStaticMesh(*PhononScene, IplVertices, IplTriangles, IplMaterialIndices, PhononStaticMesh);
+			NumSceneTriangles = IplTriangles.Num();
+
+			Result.SetValue(true);
+		});
+
+		TFuture<bool> Value = Result.GetFuture();
+		Value.Wait();
+
+		return Value.Get();
+	}
+
+	void CreateDynamicScene(UWorld* World, IPLSimulationSettings SimulationSettings, TArray<IPLMaterial> SceneMaterials, uint32& NumSceneTriangles, uint32& DataSize, bool CreateObjFiles /* =false */)
+	{
+		// Wipe old dynamic scene files for this level (if any)
+		DeleteSceneFiles(World, DynamicRuntimePath);
+		DeleteSceneFiles(World, DynamicEditorOnlyPath);
+
+		for (TActorIterator<AActor> ActorItr(World); ActorItr; ++ActorItr)
+		{
+			if (ActorItr && ActorItr->IsValidLowLevelFast())
+			{
+				// Check if the actor's root component is movable and there's a phonon geometry component and at least one static mesh component attached
+				auto PhononGeometryComponent = ActorItr->GetComponentByClass(UPhononGeometryComponent::StaticClass());
+				if (ActorItr->GetRootComponent()
+					&& ActorItr->GetRootComponent()->Mobility == EComponentMobility::Movable
+					&& PhononGeometryComponent
+					&& PhononGeometryComponent->IsValidLowLevelFast()
+					&& ActorItr->GetComponentByClass(UStaticMeshComponent::StaticClass())
+					)
+				{
+					// Go through all static mesh components in this actor
+					TArray<UActorComponent*> Components;
+					ActorItr->GetComponents(UStaticMeshComponent::StaticClass(), Components);
+
+					for (auto ActorItrComponent : Components)
+					{
+						auto DynamicMeshComponent = Cast<UStaticMeshComponent>(ActorItrComponent);
+
+						if (DynamicMeshComponent && DynamicMeshComponent->IsValidLowLevelFast())
+						{
+							UPhononGeometryComponent* DynamicGeometryComponent = nullptr;
+							TArray<USceneComponent*> ChildrenComponents;
+							DynamicMeshComponent->GetChildrenComponents(false, ChildrenComponents);
+
+							// Check if this static mesh component has a Phonon Geometry Component
+							for (auto ChildComponent : ChildrenComponents)
+							{
+								auto ChildPhononGeometryComponent = Cast<UPhononGeometryComponent>(ChildComponent);
+
+								if (ChildPhononGeometryComponent)
+								{
+									DynamicGeometryComponent = ChildPhononGeometryComponent;
+									break;
+								}
+							}
+
+							if (DynamicGeometryComponent)
+							{
+								// Create an iplScene 
+								IPLhandle DynamicPhononScene = nullptr;
+								FString SceneID = World->GetName() + "_" + ActorItr->GetName() + "_" + DynamicMeshComponent->GetName();
+
+								IPLerror IplDynResult = iplCreateScene(GlobalContext, nullptr, SimulationSettings, SceneMaterials.Num(),
+									SceneMaterials.GetData(), nullptr, nullptr, nullptr, nullptr, nullptr, &DynamicPhononScene);
+								if (IplDynResult == IPL_STATUS_SUCCESS)
+								{
+									UE_LOG(LogSteamAudio, Log, TEXT("Dynamic Phonon scene succesfully created for [%s]"), *SceneID);
+								}
+								else
+								{
+									UE_LOG(LogSteamAudio, Warning, TEXT("Error creating Dynamic Phonon scene for [%s]"), *SceneID);
+									continue;
+								}
+
+								// Create and register dynamic meshes to the dynamic scene
+								uint32 DynamicSceneTriangles = 0;
+								TArray<IPLhandle> PhononDynamicMeshes;
+								TArray<IPLVector3> DynIplVertices;
+								TArray<IPLTriangle> DynIplTriangles;
+								TArray<IPLint32> DynIplMaterialIndices;
+
+
+								CalculateVertsTrisMats(*ActorItr, DynamicMeshComponent, DynIplVertices, DynIplTriangles, DynIplMaterialIndices);
+
+								// Register Instanced Mesh to Dynamic Scene
+								RegisterStaticMeshes(DynamicPhononScene, DynIplVertices, DynIplTriangles, DynIplMaterialIndices, &PhononDynamicMeshes);
+
+								// Add tris in this mesh to the total scene tris
+								DynamicSceneTriangles += DynIplTriangles.Num();
+
+								// Finalize and save to disk
+								FPhononSceneInfo DynamicPhononSceneInfo;
+								DynamicPhononSceneInfo.NumTriangles = DynamicSceneTriangles;
+								DynamicPhononSceneInfo.DataSize = iplSaveScene(DynamicPhononScene, nullptr);
+								bool SaveDynSceneSuccessful = SaveDynamicSceneToDisk(World, DynamicPhononScene, DynamicPhononSceneInfo, SceneID, CreateObjFiles);
+
+								// Clean up Dynamic Phonon structures
+								for (IPLhandle PhononDynamicMeshItr : PhononDynamicMeshes)
+								{
+									iplDestroyStaticMesh(&PhononDynamicMeshItr);
+								}
+								iplDestroyScene(&DynamicPhononScene);
+
+								if (SaveDynSceneSuccessful)
+								{
+									DataSize += DynamicPhononSceneInfo.DataSize;
+									NumSceneTriangles += DynamicPhononSceneInfo.NumTriangles;
+									UE_LOG(LogSteamAudio, Log, TEXT("Dynamic scene [%s] with [%i] triangles saved to disk with size [%i]."), *SceneID, DynamicSceneTriangles, DynamicPhononSceneInfo.DataSize);
+								}
+								else
+								{
+									UE_LOG(LogSteamAudio, Warning, TEXT("Unable to save Dynamic scene [%s] with [%i] triangles to disk with size [%i]."), *SceneID, DynamicSceneTriangles, DynamicPhononSceneInfo.DataSize);
+								}
+							}
+						}
+					}
+				}
+			}
+		}
+	}
+
+	bool SaveFinalizedSceneToDisk(UWorld* World, IPLhandle PhononScene, const FPhononSceneInfo& PhononSceneInfo)
+	{
+		// Write Phonon Scene data to byte array
+		TArray<uint8> SceneData;
+		SceneData.SetNum(PhononSceneInfo.DataSize);
+		iplSaveScene(PhononScene, SceneData.GetData());
+
+		// Serialize byte array to disk
+		FString SceneFileName = RuntimePath + World->GetMapName() + ".phononscene";
+		FString SceneInfoName = EditorOnlyPath + World->GetMapName() + ".phononsceneinfo";
+		IPlatformFile& PlatformFile = FPlatformFileManager::Get().GetPlatformFile();
+		IFileHandle* SceneFileHandle = PlatformFile.OpenWrite(*SceneFileName);
+		IFileHandle* SceneInfoHandle = PlatformFile.OpenWrite(*SceneInfoName);
+
+		if (SceneFileHandle && SceneInfoHandle)
+		{
+			SceneFileHandle->Write(SceneData.GetData(), PhononSceneInfo.DataSize);
+			delete SceneFileHandle;
+
+			SceneInfoHandle->Write((uint8*)&PhononSceneInfo.NumTriangles, 4);
+			SceneInfoHandle->Write((uint8*)&PhononSceneInfo.NumDynTriangles, 4);
+			SceneInfoHandle->Write((uint8*)&PhononSceneInfo.DataSize, 4);
+			SceneInfoHandle->Write((uint8*)&PhononSceneInfo.DynDataSize, 4);
+			delete SceneInfoHandle;
+		}
+		else
+		{
+			if (SceneFileHandle)
+			{
+				delete SceneFileHandle;
+			}
+
+			if (SceneInfoHandle)
+			{
+				delete SceneInfoHandle;
+			}
+
+			return false;
+		}
+
+		return true;
+	}
+
+
+	bool STEAMAUDIO_API SaveDynamicSceneToDisk(UWorld* World, IPLhandle PhononScene, const FPhononSceneInfo& PhononSceneInfo, FString DynamicSceneID, bool CreateObjFile /* =false */)
+	{
+		// Write Phonon Scene data to byte array
+		TArray<uint8> SceneData;
+		SceneData.SetNum(PhononSceneInfo.DataSize);
+		iplSaveScene(PhononScene, SceneData.GetData());
+
+		// Serialize byte array to disk
+		FString SceneFileName = DynamicRuntimePath + DynamicSceneID + ".phononscene";
+		FString SceneInfoName = DynamicEditorOnlyPath + DynamicSceneID + ".phononinfo";
+		IPlatformFile& PlatformFile = FPlatformFileManager::Get().GetPlatformFile();
+		IFileHandle* SceneFileHandle = PlatformFile.OpenWrite(*SceneFileName);
+		IFileHandle* SceneInfoHandle = PlatformFile.OpenWrite(*SceneInfoName);
+
+		if (SceneFileHandle && SceneInfoHandle)
+		{
+			if (CreateObjFile)
+			{
+				iplSaveSceneAsObj(PhononScene, TCHAR_TO_ANSI(*(DynamicEditorOnlyPath + DynamicSceneID + ".obj")));
+			}
+
+			SceneFileHandle->Write(SceneData.GetData(), PhononSceneInfo.DataSize);
+			delete SceneFileHandle;
+
+			SceneInfoHandle->Write((uint8*)&PhononSceneInfo.NumTriangles, 4);
+			SceneInfoHandle->Write((uint8*)&PhononSceneInfo.DataSize, 4);
+			delete SceneInfoHandle;
+		}
+		else
+		{
+			if (SceneFileHandle)
+			{
+				delete SceneFileHandle;
+			}
+
+			if (SceneInfoHandle)
+			{
+				delete SceneInfoHandle;
+			}
+
+			return false;
+		}
+
+		return true;
+	}
+
+	bool STEAMAUDIO_API CheckSceneDirectory(FString DirectoryPath)
+	{
+		IFileManager& FileManager = FFileManagerGeneric::Get();
+		if (!FileManager.DirectoryExists(*DirectoryPath))
+		{
+			FileManager.MakeDirectory(*DirectoryPath);
+			return false;
+		}
+		
+		return true;
+	}
+
+	void STEAMAUDIO_API DeleteSceneFiles(UWorld * World, FString FilePath)
+	{
+		IFileManager& FileManager = FFileManagerGeneric::Get();
+		if (CheckSceneDirectory(FilePath))
+		{
+			TArray<FString> FoundFiles;
+			FileManager.FindFiles(FoundFiles, *FilePath);
+
+			for (auto FoundFile : FoundFiles)
+			{
+				FString FileToDelete = FilePath + FoundFile;
+				//UE_LOG(LogSteamAudio, Warning, TEXT("Checking for file [%s]"), *FileToDelete);
+
+				if (FoundFile.Left(World->GetName().Len()).Equals(World->GetName(), ESearchCase::IgnoreCase) 
+					&& FileManager.Delete(*FileToDelete, true, true, true)
+					)
+				{
+					UE_LOG(LogSteamAudio, Log, TEXT("Deleted old scene file [%s]"), *FileToDelete);
+				}
+			}
+		}
+	}
+
+	//==============================================================================================================================================
+	// Utilities for adding/removing Phonon Geometry components
+	//==============================================================================================================================================
+
+	/**
+	 * Adds Phonon Geometry components with default settings. Will not add if one already exists.
+	 */
+	void AddGeometryComponentsToStaticMeshes(UWorld* World)
+	{
+		check(World);
+
+		for (TActorIterator<AActor> ActorItr(World); ActorItr; ++ActorItr)
+		{
+			
+			// Go through any static mesh components found in this actor
+			TArray<UActorComponent*> Components;
+			ActorItr->GetComponents(UStaticMeshComponent::StaticClass(), Components);
+
+			for (auto ActorComp : Components)
+			{
+				if (ActorComp && ActorComp->IsValidLowLevel())
+				{
+					auto* StaticMeshComp = Cast<UStaticMeshComponent>(ActorComp);
+	
+					// Check if this static mesh can affect audio
+					if (StaticMeshComp
+						&& StaticMeshComp->IsValidLowLevel()
+						&& !StaticMeshComp->IsVisualizationComponent()
+						)
+					{
+						// Check if there's already a phonon geometry component attached to this static mesh component
+						TArray<USceneComponent*> ChildrenComponents;
+						StaticMeshComp->GetChildrenComponents(true, ChildrenComponents);
+						bool bAlreadyHavePhononGeomComponent = false;
+						
+						for (auto ChildComponent : ChildrenComponents)
+						{
+							auto* PhononGeomComp = Cast<UPhononGeometryComponent>(ChildComponent);
+							
+							if (PhononGeomComp && PhononGeomComp->IsValidLowLevelFast())
+							{
+								bAlreadyHavePhononGeomComponent = true;
+								break;
+							}
+						}
+
+						// If there isn't, create and attach a new phonon geometry component to this static mesh
+						if (!bAlreadyHavePhononGeomComponent)
+						{
+							// Add phonon geometry object
+							auto PhononGeometryComponent = NewObject<UPhononGeometryComponent>(*ActorItr);
+							PhononGeometryComponent->AttachToComponent(StaticMeshComp, FAttachmentTransformRules::KeepRelativeTransform);
+							PhononGeometryComponent->RegisterComponent();
+							ActorItr->AddInstanceComponent(PhononGeometryComponent);
+						}
+					}
+				}
+			}
+		}
+	}
+
+	/**
+	 * Removes all Phonon Geometry components from Static Mesh actors.
+	 */
+	void RemoveGeometryComponentsFromStaticMeshes(UWorld* World)
+	{
+		check(World);
+
+		for (TObjectIterator<UPhononGeometryComponent> PhononGeometryObj; PhononGeometryObj; ++PhononGeometryObj)
+		{
+			if (PhononGeometryObj && PhononGeometryObj->IsValidLowLevel())
+			{
+				PhononGeometryObj->DestroyComponent();
+			}
+		}
+	}
+
+
+	//==============================================================================================================================================
+	// Static mesh geometry export
+	//==============================================================================================================================================
+
+	/**
+	 * Populates VertexArray with the given mesh's vertices. Converts from UE coords to Phonon coords. Returns the number of vertices added.
+	 */
+	static uint32 GetMeshVerts(UStaticMeshComponent* StaticMeshComponent, TArray<IPLVector3>& VertexArray, bool bGetRelativeVertexPos /*= true*/)
+	{
+		check(StaticMeshComponent->GetStaticMesh()->HasValidRenderData());
+
+		uint32 NumVerts = 0;
+
+		FStaticMeshLODResources& LODModel = StaticMeshComponent->GetStaticMesh()->GetRenderData()->LODResources[0];
+		auto Indices = LODModel.IndexBuffer.GetArrayView();
+
+		for (const FStaticMeshSection& Section : LODModel.Sections)
+		{
+			for (auto TriIndex = 0; TriIndex < (int32)Section.NumTriangles; ++TriIndex)
+			{
+				auto BaseIndex = Section.FirstIndex + TriIndex * 3;
+				for (auto v = 2; v > -1; v--)
+				{
+					auto i = Indices[BaseIndex + v];
+					auto vtx = bGetRelativeVertexPos ? FVector3f(StaticMeshComponent->GetComponentTransform().TransformPosition((FVector)LODModel.VertexBuffers.PositionVertexBuffer.VertexPosition(i)))
+						: LODModel.VertexBuffers.PositionVertexBuffer.VertexPosition(i);
+
+					VertexArray.Add(UnrealToPhononIPLVector3((FVector)vtx));
+					NumVerts++;
+				}
+			}
+		}
+
+		return NumVerts;
+	}
+
+	/**
+	 * Walks up the actor attachment chain, checking for a Phonon Geometry component.
+	 */
+	static bool IsActorPhononGeometry(AActor* Actor)
+	{
+		auto CurrentActor = Actor;
+		while (CurrentActor)
+		{
+			if (CurrentActor->GetComponentByClass(UPhononGeometryComponent::StaticClass()))
+			{
+				return true;
+			}
+			CurrentActor = CurrentActor->GetAttachParentActor();
+		}
+
+		return false;
+	}
+
+	/**
+	 * Walks up the actor attachment chain, checking for a Phonon Material component.
+	 */
+	static UPhononMaterialComponent* GetPhononMaterialComponent(AActor* Actor)
+	{
+		auto CurrentActor = Actor;
+		while (CurrentActor)
+		{
+			if (CurrentActor->GetComponentByClass(UPhononMaterialComponent::StaticClass()))
+			{
+				return static_cast<UPhononMaterialComponent*>(CurrentActor->GetComponentByClass(UPhononMaterialComponent::StaticClass()));
+			}
+			CurrentActor = CurrentActor->GetAttachParentActor();
+		}
+
+		return nullptr;
+	}
+
+	/**
+	 * Calculates a mesh's vertices, tris and retries phonon material indices for a given actor
+	 */
+	static void CalculateVertsTrisMats(AActor* MeshActor, UStaticMeshComponent* MeshComponent, TArray<IPLVector3>& IplVertices, TArray<IPLTriangle>& IplTriangles,
+		TArray<IPLint32>& IplMaterialIndices)
+	{
+		check(MeshActor);
+		check(MeshComponent);
+
+		int32 StartVertexIdx = IplVertices.Num();
+		int32 NumMeshVerts = MeshComponent->Mobility == EComponentMobility::Movable ? GetMeshVerts(MeshComponent, IplVertices, false) : GetMeshVerts(MeshComponent, IplVertices);
+		int32 NumMeshTriangles = NumMeshVerts / 3;
+
+		for (int32 i = 0; i < NumMeshTriangles; ++i)
+		{
+			IPLTriangle IplTriangle;
+			IplTriangle.indices[0] = StartVertexIdx + i * 3;
+			IplTriangle.indices[1] = StartVertexIdx + i * 3 + 2;
+			IplTriangle.indices[2] = StartVertexIdx + i * 3 + 1;
+			IplTriangles.Add(IplTriangle);
+		}
+
+		auto PhononMaterialComponent = GetPhononMaterialComponent(MeshActor);
+		auto MaterialIndex = 0;
+
+		if (PhononMaterialComponent)
+		{
+			MaterialIndex = PhononMaterialComponent->MaterialIndex;
+		}
+		else
+		{
+			// The default static mesh material is always registered at size(MaterialPresets)
+			MaterialIndex = MaterialPresets.Num();
+		}
+
+		for (auto i = 0; i < NumMeshTriangles; ++i)
+		{
+			IplMaterialIndices.Add(MaterialIndex);
+			//UE_LOG(LogSteamAudio, Warning, TEXT("Phonon scene Material Index Added [%i]"), MaterialIndex);
+		}
+	}
+
+	/**
+	 * Loads any static mesh actors.
+	 */
+	static void LoadStaticMeshActors(UWorld* World, IPLhandle PhononScene, TArray<IPLVector3>& IplVertices, TArray<IPLTriangle>& IplTriangles,
+		TArray<IPLint32>& IplMaterialIndices)
+	{
+		check(World);
+		check(PhononScene);
+
+		UE_LOG(LogSteamAudio, Log, TEXT("Loading static mesh actors."));
+
+		for (TActorIterator<AStaticMeshActor> AStaticMeshItr(World); AStaticMeshItr; ++AStaticMeshItr)
+		{
+			// Only consider static mesh actors that have both an acoustic geometry component attached and valid render data
+			if (IsActorPhononGeometry(*AStaticMeshItr) 
+				&& AStaticMeshItr->GetStaticMeshComponent()->Mobility != EComponentMobility::Movable
+				&& AStaticMeshItr->GetStaticMeshComponent()->GetStaticMesh() 
+				&& AStaticMeshItr->GetStaticMeshComponent()->GetStaticMesh()->HasValidRenderData()
+				)
+			{
+				CalculateVertsTrisMats(*AStaticMeshItr, AStaticMeshItr->GetStaticMeshComponent(), IplVertices, IplTriangles, IplMaterialIndices);
+			}
+		}
+	}
+
+	/**
+	 * Loads any static mesh actors, adding any Phonon static meshes to the provided array.
+	 */
+	static uint32 LoadDynamicMeshes(AActor* DynamicActor, IPLhandle PhononScene)
+	{
+		check(DynamicActor);
+		check(PhononScene);
+
+		TArray<IPLVector3> IplVertices;
+		TArray<IPLTriangle> IplTriangles;
+		TArray<IPLint32> IplMaterialIndices;
+
+		// Go through all static mesh components in this actor
+		TArray<UActorComponent*> Components;
+		DynamicActor->GetComponents(UStaticMeshComponent::StaticClass(), Components);
+		for (auto StaticMeshComp : Components)
+		{
+			auto StaticMeshComponent = Cast<UStaticMeshComponent>(StaticMeshComp);
+
+			// Only consider dynamic static meshes that have a static mesh already assigned and have valid render data
+			if (StaticMeshComponent
+				&& StaticMeshComponent->IsValidLowLevelFast()
+				&& StaticMeshComponent->Mobility == EComponentMobility::Movable
+				&& StaticMeshComponent->GetStaticMesh()
+				&& StaticMeshComponent->GetStaticMesh()->HasValidRenderData()
+				&& !StaticMeshComponent->IsVisualizationComponent()
+				)
+			{
+				CalculateVertsTrisMats(DynamicActor, StaticMeshComponent, IplVertices, IplTriangles, IplMaterialIndices);
+			}
+		}
+
+		// Register a new static mesh with Phonon
+		//RegisterStaticMesh(PhononScene, IplVertices, IplTriangles, IplMaterialIndices, PhononStaticMeshes);
+
+		return IplTriangles.Num();
+	}
+
+	//==============================================================================================================================================
+	// BSP geometry export
+	//==============================================================================================================================================
+
+	/**
+	 * Loads any BSP geometry.
+	 */
+	static void LoadBSPGeometry(UWorld* World, IPLhandle PhononScene, TArray<IPLVector3>& IplVertices, TArray<IPLTriangle>& IplTriangles,
+		TArray<IPLint32>& IplMaterialIndices)
+	{
+		check(World);
+		check(PhononScene);
+
+		UE_LOG(LogSteamAudio, Log, TEXT("Loading BSP geometry."));
+
+		int32 InitialNumVertices = IplVertices.Num();
+		int32 InitialNumTriangles = IplTriangles.Num();
+
+		// Gather and convert all world vertices to Phonon coords
+		for (auto& WorldVertex : World->GetModel()->Points)
+		{
+			IplVertices.Add(SteamAudio::UnrealToPhononIPLVector3((FVector)WorldVertex));
+		}
+
+		// Gather vertex indices for all faces ("nodes" are faces)
+		for (auto& WorldNode : World->GetModel()->Nodes)
+		{
+			// Ignore degenerate faces
+			if (WorldNode.NumVertices <= 2)
+			{
+				continue;
+			}
+
+			// Faces are organized as triangle fans
+			int32 Index0 = World->GetModel()->Verts[WorldNode.iVertPool + 0].pVertex;
+			int32 Index1 = World->GetModel()->Verts[WorldNode.iVertPool + 1].pVertex;
+			int32 Index2;
+
+			for (auto v = 2; v < WorldNode.NumVertices; ++v)
+			{
+				Index2 = World->GetModel()->Verts[WorldNode.iVertPool + v].pVertex;
+
+				IPLTriangle IplTriangle;
+				IplTriangle.indices[0] = Index0 + InitialNumVertices;
+				IplTriangle.indices[1] = Index2 + InitialNumVertices;
+				IplTriangle.indices[2] = Index1 + InitialNumVertices;
+				IplTriangles.Add(IplTriangle);
+
+				Index1 = Index2;
+			}
+		}
+
+		// The default BSP material is always registered at size(MaterialPresets) + 1
+		auto MaterialIdx = MaterialPresets.Num() + 1;
+		for (auto i = InitialNumTriangles; i < IplTriangles.Num(); ++i)
+		{
+			IplMaterialIndices.Add(MaterialIdx);
+		}
+	}
+
+	//==============================================================================================================================================
+	// Landscape geometry export
+	//==============================================================================================================================================
+
+	/**
+	 * Loads any Landscape actors.
+	 */
+	static void LoadLandscapeActors(UWorld* World, IPLhandle PhononScene, TArray<IPLVector3>& IplVertices, TArray<IPLTriangle>& IplTriangles,
+		TArray<IPLint32>& IplMaterialIndices)
+	{
+		check(World);
+		check(PhononScene);
+
+		UE_LOG(LogSteamAudio, Log, TEXT("Loading landscape actors."));
+
+		int32 InitialNumTriangles = IplTriangles.Num();
+
+		for (TActorIterator<ALandscape> ALandscapeItr(World); ALandscapeItr; ++ALandscapeItr)
+		{
+			ULandscapeInfo* LandscapeInfo = ALandscapeItr->GetLandscapeInfo();
+
+			for (auto It = LandscapeInfo->XYtoComponentMap.CreateIterator(); It; ++It)
+			{
+				ULandscapeComponent* Component = It.Value();
+				FLandscapeComponentDataInterface CDI(Component);
+
+				for (auto y = 0; y < Component->ComponentSizeQuads; ++y)
+				{
+					for (auto x = 0; x < Component->ComponentSizeQuads; ++x)
+					{
+						auto StartIndex = IplVertices.Num();
+
+						IplVertices.Add(SteamAudio::UnrealToPhononIPLVector3(CDI.GetWorldVertex(x, y)));
+						IplVertices.Add(SteamAudio::UnrealToPhononIPLVector3(CDI.GetWorldVertex(x, y + 1)));
+						IplVertices.Add(SteamAudio::UnrealToPhononIPLVector3(CDI.GetWorldVertex(x + 1, y + 1)));
+						IplVertices.Add(SteamAudio::UnrealToPhononIPLVector3(CDI.GetWorldVertex(x + 1, y)));
+
+						IPLTriangle Triangle;
+
+						Triangle.indices[0] = StartIndex;
+						Triangle.indices[1] = StartIndex + 2;
+						Triangle.indices[2] = StartIndex + 3;
+						IplTriangles.Add(Triangle);
+
+						Triangle.indices[0] = StartIndex;
+						Triangle.indices[1] = StartIndex + 1;
+						Triangle.indices[2] = StartIndex + 2;
+						IplTriangles.Add(Triangle);
+					}
+				}
+			}
+		}
+
+		// The default landscape material is always registered at size(MaterialPresets) + 2
+		auto MaterialIdx = MaterialPresets.Num() + 2;
+		for (auto i = InitialNumTriangles; i < IplTriangles.Num(); ++i)
+		{
+			IplMaterialIndices.Add(MaterialIdx);
+		}
+	}
+
+	//==============================================================================================================================================
+	// Utility functions
+	//==============================================================================================================================================
+
+	/**
+	 * Registers a new static mesh with Phonon.
+	 */
+	static void RegisterStaticMesh(IPLhandle PhononScene, TArray<IPLVector3>& IplVertices, TArray<IPLTriangle>& IplTriangles,
+		TArray<IPLint32>& IplMaterialIndices, IPLhandle* PhononStaticMesh)
+	{
+		if (IplVertices.Num() > 0)
+		{
+			UE_LOG(LogSteamAudio, Log, TEXT("Registering new mesh with %d triangles."), IplTriangles.Num());
+
+			auto IplResult = iplCreateStaticMesh(PhononScene, IplVertices.Num(), IplTriangles.Num(), IplVertices.GetData(), IplTriangles.GetData(),
+				IplMaterialIndices.GetData(), PhononStaticMesh);
+			if (IplResult != IPL_STATUS_SUCCESS)
+			{
+				UE_LOG(LogSteamAudio, Warning, TEXT("Error adding a new object to the Phonon scene."));
+				return;
+			}
+		}
+		else
+		{
+			UE_LOG(LogSteamAudio, Warning, TEXT("Skipping mesh registration because no vertices were found."));
+		}
+	}
+
+	/**
+	 * Registers a new static mesh with Phonon, adding its handle to the provided array of static meshes.
+	 */
+	static void RegisterStaticMeshes(IPLhandle PhononScene, TArray<IPLVector3>& IplVertices, TArray<IPLTriangle>& IplTriangles,
+		TArray<IPLint32>& IplMaterialIndices, TArray<IPLhandle>* PhononStaticMeshes)
+	{
+		if (IplVertices.Num() > 0)
+		{
+			UE_LOG(LogSteamAudio, Log, TEXT("Registering new mesh with %d triangles."), IplTriangles.Num());
+
+			IPLhandle IplStaticMesh = nullptr;
+			auto IplResult = iplCreateStaticMesh(PhononScene, IplVertices.Num(), IplTriangles.Num(), IplVertices.GetData(), IplTriangles.GetData(),
+				IplMaterialIndices.GetData(), &IplStaticMesh);
+			if (IplResult != IPL_STATUS_SUCCESS)
+			{
+				UE_LOG(LogSteamAudio, Warning, TEXT("Error adding a new object to the acoustic scene."));
+				return;
+			}
+
+			PhononStaticMeshes->Add(IplStaticMesh);
+		}
+		else
+		{
+			UE_LOG(LogSteamAudio, Warning, TEXT("Skipping mesh registration because no vertices were found."));
+		}
+	}
+
+	/**
+	 * Calculates the total number of materials that must be registered with Phonon. This includes presets and any custom materials.
+	 */
+	static void GetPhononMaterials(UWorld* World, TArray<IPLMaterial>& IPLMaterials)
+	{
+		check(World);
+
+		// There are size(MaterialPresets) + 3 fixed slots.
+		uint32 NumMaterials = MaterialPresets.Num() + 3;
+
+		for (TActorIterator<AActor> AActorItr(World); AActorItr; ++AActorItr)
+		{
+			auto PhononMaterialComponent = static_cast<UPhononMaterialComponent*>(
+				AActorItr->GetComponentByClass(UPhononMaterialComponent::StaticClass()));
+			if (PhononMaterialComponent)
+			{
+
+				if (PhononMaterialComponent->MaterialPreset == EPhononMaterial::CUSTOM)
+				{
+					PhononMaterialComponent->MaterialIndex = NumMaterials++;
+					IPLMaterials.Add(PhononMaterialComponent->GetMaterialPreset());
+					UE_LOG(LogSteamAudio, Log, TEXT("Custom Material [%i] added with Absorption[%f][%f][%f] Transmission[%f][%f][%f]"), 
+						PhononMaterialComponent->MaterialIndex,
+						PhononMaterialComponent->GetMaterialPreset().lowFreqAbsorption,
+						PhononMaterialComponent->GetMaterialPreset().midFreqAbsorption,
+						PhononMaterialComponent->GetMaterialPreset().highFreqAbsorption,
+						PhononMaterialComponent->GetMaterialPreset().lowFreqTransmission,
+						PhononMaterialComponent->GetMaterialPreset().midFreqTransmission,
+						PhononMaterialComponent->GetMaterialPreset().highFreqTransmission);
+				}
+				else
+				{
+					PhononMaterialComponent->MaterialIndex = static_cast<int32>(PhononMaterialComponent->MaterialPreset);
+				}
+			}
+		}
+	}
+
+#endif // WITH_EDITOR
+
+	uint32 GetNumTrianglesForStaticMesh(AStaticMeshActor* StaticMeshActor)
+	{
+		uint32 NumTriangles = 0;
+
+		if (StaticMeshActor == nullptr || StaticMeshActor->GetStaticMeshComponent() == nullptr ||
+			StaticMeshActor->GetStaticMeshComponent()->GetStaticMesh() == nullptr ||
+			StaticMeshActor->GetStaticMeshComponent()->GetStaticMesh()->GetRenderData() == nullptr)
+		{
+			return NumTriangles;
+		}
+
+		const auto& LODModel = StaticMeshActor->GetStaticMeshComponent()->GetStaticMesh()->GetRenderData()->LODResources[0];
+		for (const auto& Section : LODModel.Sections)
+		{
+			NumTriangles += Section.NumTriangles;
+		}
+
+		return NumTriangles;
+	}
+
+	uint32 GetNumTrianglesAtRoot(AActor* RootActor)
+	{
+		uint32 NumTriangles = 0;
+
+		if (RootActor == nullptr)
+		{
+			return NumTriangles;
+		}
+
+		NumTriangles = GetNumTrianglesForStaticMesh(Cast<AStaticMeshActor>(RootActor));
+
+		TArray<AActor*> AttachedActors;
+		RootActor->GetAttachedActors(AttachedActors);
+
+		for (auto AttachedActor : AttachedActors)
+		{
+			NumTriangles += GetNumTrianglesAtRoot(AttachedActor);
+		}
+
+		return NumTriangles;
+	}
+}
diff --git a/Engine/Plugins/Runtime/Steam/SteamAudio/Source/SteamAudio/Private/PhononSourceComponent.cpp b/Engine/Plugins/Runtime/Steam/SteamAudio/Source/SteamAudio/Private/PhononSourceComponent.cpp
new file mode 100644
--- /dev/null
+++ b/Engine/Plugins/Runtime/Steam/SteamAudio/Source/SteamAudio/Private/PhononSourceComponent.cpp
@@ -0,0 +1,5 @@
+//
+// Copyright (C) Valve Corporation. All rights reserved.
+//
+
+#include "PhononSourceComponent.h"
\ No newline at end of file
diff --git a/Engine/Plugins/Runtime/Steam/SteamAudio/Source/SteamAudio/Private/PhononSpatialization.cpp b/Engine/Plugins/Runtime/Steam/SteamAudio/Source/SteamAudio/Private/PhononSpatialization.cpp
new file mode 100644
--- /dev/null
+++ b/Engine/Plugins/Runtime/Steam/SteamAudio/Source/SteamAudio/Private/PhononSpatialization.cpp
@@ -0,0 +1,180 @@
+//
+// Copyright (C) Valve Corporation. All rights reserved.
+//
+
+#include "PhononSpatialization.h"
+#include "PhononSpatializationSourceSettings.h"
+#include "PhononCommon.h"
+#include "SteamAudioSettings.h"
+
+namespace SteamAudio
+{
+	//==============================================================================================================================================
+	// FPhononSpatialization
+	//==============================================================================================================================================
+
+	FPhononSpatialization::FPhononSpatialization()
+		: BinauralRenderer(nullptr)
+	{
+		InputAudioFormat.channelLayout = IPL_CHANNELLAYOUT_MONO;
+		InputAudioFormat.channelLayoutType = IPL_CHANNELLAYOUTTYPE_SPEAKERS;
+		InputAudioFormat.channelOrder = IPL_CHANNELORDER_INTERLEAVED;
+		InputAudioFormat.numSpeakers = 1;
+		InputAudioFormat.speakerDirections = nullptr;
+		InputAudioFormat.ambisonicsOrder = -1;
+		InputAudioFormat.ambisonicsNormalization = IPL_AMBISONICSNORMALIZATION_N3D;
+		InputAudioFormat.ambisonicsOrdering = IPL_AMBISONICSORDERING_ACN;
+
+		BinauralOutputAudioFormat.channelLayout = IPL_CHANNELLAYOUT_STEREO;
+		BinauralOutputAudioFormat.channelLayoutType = IPL_CHANNELLAYOUTTYPE_SPEAKERS;
+		BinauralOutputAudioFormat.channelOrder = IPL_CHANNELORDER_INTERLEAVED;
+		BinauralOutputAudioFormat.numSpeakers = 2;
+		BinauralOutputAudioFormat.speakerDirections = nullptr;
+		BinauralOutputAudioFormat.ambisonicsOrder = -1;
+		BinauralOutputAudioFormat.ambisonicsNormalization = IPL_AMBISONICSNORMALIZATION_N3D;
+		BinauralOutputAudioFormat.ambisonicsOrdering = IPL_AMBISONICSORDERING_ACN;
+	}
+
+	FPhononSpatialization::~FPhononSpatialization()
+	{
+		if (BinauralRenderer)
+		{
+			iplDestroyBinauralRenderer(&BinauralRenderer);
+		}
+	}
+
+	void FPhononSpatialization::Initialize(const FAudioPluginInitializationParams InitializationParams)
+	{
+		// The binaural renderer always uses Phonon convolution irrespective of whether TAN is available.
+		RenderingSettings.convolutionType = IPL_CONVOLUTIONTYPE_PHONON;
+		RenderingSettings.frameSize = InitializationParams.BufferLength;
+		RenderingSettings.samplingRate = InitializationParams.SampleRate;
+
+		IPLHrtfParams HrtfParams;
+		HrtfParams.hrtfData = nullptr;
+		// TODO: Sofa File Support
+		HrtfParams.sofaFileName = IPLstring("");
+		HrtfParams.type = IPL_HRTFDATABASETYPE_DEFAULT;
+
+		iplCreateBinauralRenderer(SteamAudio::GlobalContext, RenderingSettings, HrtfParams, &BinauralRenderer);
+
+		BinauralSources.AddDefaulted(InitializationParams.NumSources);
+		for (FBinauralSource& BinauralSource : BinauralSources)
+		{
+			BinauralSource.InBuffer.format = InputAudioFormat;
+			BinauralSource.InBuffer.numSamples = InitializationParams.BufferLength;
+			BinauralSource.InBuffer.interleavedBuffer = nullptr;
+			BinauralSource.InBuffer.deinterleavedBuffer = nullptr;
+
+			BinauralSource.OutArray.SetNumZeroed(InitializationParams.BufferLength * 2);
+			BinauralSource.OutBuffer.format = BinauralOutputAudioFormat;
+			BinauralSource.OutBuffer.numSamples = InitializationParams.BufferLength;
+			BinauralSource.OutBuffer.interleavedBuffer = nullptr;
+			BinauralSource.OutBuffer.deinterleavedBuffer = nullptr;
+		}
+	}
+
+	void FPhononSpatialization::OnInitSource(const uint32 SourceId, const FName& AudioComponentUserId,
+		USpatializationPluginSourceSettingsBase* InSettings)
+	{
+		UPhononSpatializationSourceSettings* SpatializationSettings = static_cast<UPhononSpatializationSourceSettings*>(InSettings);
+		FBinauralSource& BinauralSource = BinauralSources[SourceId];
+
+		UE_LOG(LogSteamAudio, Verbose, TEXT("Creating spatialization effect."));
+
+		if (SpatializationSettings)
+		{
+			switch (SpatializationSettings->SpatializationMethod)
+			{
+			case EIplSpatializationMethod::HRTF:
+				iplCreateBinauralEffect(BinauralRenderer, InputAudioFormat, BinauralOutputAudioFormat, &BinauralSource.BinauralEffect);
+				break;
+			case EIplSpatializationMethod::PANNING:
+				iplCreatePanningEffect(BinauralRenderer, InputAudioFormat, BinauralOutputAudioFormat, &BinauralSource.PanningEffect);
+				break;
+			}
+
+			BinauralSource.SpatializationMethod = SpatializationSettings->SpatializationMethod;
+			BinauralSource.HrtfInterpolationMethod = SpatializationSettings->HrtfInterpolationMethod;
+		}
+		else
+		{
+			iplCreateBinauralEffect(BinauralRenderer, InputAudioFormat, BinauralOutputAudioFormat, &BinauralSource.BinauralEffect);
+			BinauralSource.SpatializationMethod = EIplSpatializationMethod::HRTF;
+			BinauralSource.HrtfInterpolationMethod = EIplHrtfInterpolationMethod::NEAREST;
+		}
+	}
+
+	void FPhononSpatialization::OnReleaseSource(const uint32 SourceId)
+	{
+		FBinauralSource& BinauralSource = BinauralSources[SourceId];
+
+		UE_LOG(LogSteamAudio, Verbose, TEXT("Destroying spatialization effect."));
+
+		switch (BinauralSource.SpatializationMethod)
+		{
+		case EIplSpatializationMethod::HRTF:
+			iplDestroyBinauralEffect(&BinauralSource.BinauralEffect);
+			break;
+		case EIplSpatializationMethod::PANNING:
+			iplDestroyPanningEffect(&BinauralSource.PanningEffect);
+			break;
+		}
+	}
+
+	void FPhononSpatialization::ProcessAudio(const FAudioPluginSourceInputData& InputData, FAudioPluginSourceOutputData& OutputData)
+	{
+		FBinauralSource& BinauralSource = BinauralSources[InputData.SourceId];
+		BinauralSource.InBuffer.interleavedBuffer = InputData.AudioBuffer->GetData();
+		BinauralSource.OutBuffer.interleavedBuffer = OutputData.AudioBuffer.GetData();
+
+		// Workaround. The directions passed to spatializer is not consistent with the
+		// coordinate system of UE, therefore special tranformation is performed here.
+		// Review this change if further changes are made to the direction passed to the
+		// spatializer.
+		IPLVector3 RelativeDirection;
+		RelativeDirection.x = InputData.SpatializationParams->EmitterPosition.Y;
+		RelativeDirection.y = InputData.SpatializationParams->EmitterPosition.X;
+		RelativeDirection.z = InputData.SpatializationParams->EmitterPosition.Z;
+
+		switch (BinauralSource.SpatializationMethod)
+		{
+		case EIplSpatializationMethod::HRTF:
+			iplApplyBinauralEffect(BinauralSource.BinauralEffect, BinauralRenderer, BinauralSource.InBuffer, RelativeDirection, static_cast<IPLHrtfInterpolation>(BinauralSource.HrtfInterpolationMethod), BinauralSource.OutBuffer);
+			break;
+		case EIplSpatializationMethod::PANNING:
+			iplApplyPanningEffect(BinauralSource.PanningEffect, BinauralRenderer, BinauralSource.InBuffer, RelativeDirection, BinauralSource.OutBuffer);
+			break;
+		}
+	}
+
+	bool FPhononSpatialization::IsSpatializationEffectInitialized() const
+	{
+		return true;
+	}
+
+	//==============================================================================================================================================
+	// FBinauralSource
+	//==============================================================================================================================================
+
+	FBinauralSource::FBinauralSource()
+		: BinauralEffect(nullptr)
+		, PanningEffect(nullptr)
+		, SpatializationMethod(EIplSpatializationMethod::HRTF)
+		, HrtfInterpolationMethod(EIplHrtfInterpolationMethod::NEAREST)
+	{
+	}
+
+	FBinauralSource::~FBinauralSource()
+	{
+		if (BinauralEffect)
+		{
+			iplDestroyBinauralEffect(&BinauralEffect);
+		}
+
+		if (PanningEffect)
+		{
+			iplDestroyPanningEffect(&PanningEffect);
+		}
+	}
+}
diff --git a/Engine/Plugins/Runtime/Steam/SteamAudio/Source/SteamAudio/Private/PhononSpatialization.h b/Engine/Plugins/Runtime/Steam/SteamAudio/Source/SteamAudio/Private/PhononSpatialization.h
new file mode 100644
--- /dev/null
+++ b/Engine/Plugins/Runtime/Steam/SteamAudio/Source/SteamAudio/Private/PhononSpatialization.h
@@ -0,0 +1,50 @@
+//
+// Copyright (C) Valve Corporation. All rights reserved.
+//
+
+#pragma once
+
+#include "IAudioExtensionPlugin.h"
+#include "phonon.h"
+#include "PhononCommon.h"
+
+namespace SteamAudio
+{
+	struct FBinauralSource
+	{
+		FBinauralSource();
+		~FBinauralSource();
+
+		IPLhandle BinauralEffect;
+		IPLhandle PanningEffect;
+		EIplSpatializationMethod SpatializationMethod;
+		EIplHrtfInterpolationMethod HrtfInterpolationMethod;
+		IPLAudioBuffer InBuffer;
+		IPLAudioBuffer OutBuffer;
+		TArray<float> OutArray;
+	};
+
+	/************************************************************************/
+	/* FPhononSpatialization                                                */
+	/* Spatialization plugin using Steam Audio's HRTF spatialization library*/
+	/************************************************************************/
+	class FPhononSpatialization : public IAudioSpatialization
+	{
+	public:
+		FPhononSpatialization();
+		~FPhononSpatialization();
+
+		virtual void Initialize(const FAudioPluginInitializationParams InitializationParams) override;
+		virtual bool IsSpatializationEffectInitialized() const override;
+		virtual void OnInitSource(const uint32 SourceId, const FName& AudioComponentUserId, USpatializationPluginSourceSettingsBase* InSettings) override;
+		virtual void OnReleaseSource(const uint32 SourceId) override;
+		virtual void ProcessAudio(const FAudioPluginSourceInputData& InputData, FAudioPluginSourceOutputData& OutputData) override;
+
+	private:
+		IPLAudioFormat InputAudioFormat;
+		IPLAudioFormat BinauralOutputAudioFormat;
+		IPLhandle BinauralRenderer;
+		IPLRenderingSettings RenderingSettings;
+		TArray<FBinauralSource> BinauralSources;
+	};
+}
diff --git a/Engine/Plugins/Runtime/Steam/SteamAudio/Source/SteamAudio/Private/PhononSpatializationSourceSettings.cpp b/Engine/Plugins/Runtime/Steam/SteamAudio/Source/SteamAudio/Private/PhononSpatializationSourceSettings.cpp
new file mode 100644
--- /dev/null
+++ b/Engine/Plugins/Runtime/Steam/SteamAudio/Source/SteamAudio/Private/PhononSpatializationSourceSettings.cpp
@@ -0,0 +1,24 @@
+//
+// Copyright (C) Valve Corporation. All rights reserved.
+//
+
+#include "PhononSpatializationSourceSettings.h"
+
+UPhononSpatializationSourceSettings::UPhononSpatializationSourceSettings()
+	: SpatializationMethod(EIplSpatializationMethod::HRTF)
+	, HrtfInterpolationMethod(EIplHrtfInterpolationMethod::NEAREST)
+{}
+
+#if WITH_EDITOR
+bool UPhononSpatializationSourceSettings::CanEditChange(const FProperty* InProperty) const
+{
+	const bool ParentVal = Super::CanEditChange(InProperty);
+
+	if ((InProperty->GetFName() == GET_MEMBER_NAME_CHECKED(UPhononSpatializationSourceSettings, HrtfInterpolationMethod)))
+	{
+		return ParentVal && SpatializationMethod == EIplSpatializationMethod::HRTF;
+	}
+
+	return ParentVal;
+}
+#endif
\ No newline at end of file
diff --git a/Engine/Plugins/Runtime/Steam/SteamAudio/Source/SteamAudio/Private/SteamAudioEnvironment.cpp b/Engine/Plugins/Runtime/Steam/SteamAudio/Source/SteamAudio/Private/SteamAudioEnvironment.cpp
new file mode 100644
--- /dev/null
+++ b/Engine/Plugins/Runtime/Steam/SteamAudio/Source/SteamAudio/Private/SteamAudioEnvironment.cpp
@@ -0,0 +1,331 @@
+//
+// Copyright (C) Valve Corporation. All rights reserved.
+//
+
+#include "SteamAudioEnvironment.h"
+#include "PhononScene.h"
+#include "PhononProbeVolume.h"
+#include "CoreMinimal.h"
+#include "Stats/Stats.h"
+#include "Modules/ModuleInterface.h"
+#include "Modules/ModuleManager.h"
+#include "Kismet/GameplayStatics.h"
+#include "Interfaces/IPluginManager.h"
+#include "EngineUtils.h"
+#include "Internationalization/Regex.h"
+#include "Misc/ScopeLock.h"
+#include "HAL/PlatformProcess.h"
+#include "Engine/StreamableManager.h"
+#include "AudioDevice.h"
+#include "PhononScene.h"
+#include "SteamAudioSettings.h"
+#include "PhononProbeVolume.h"
+#include "HAL/PlatformFileManager.h"
+#include "GenericPlatform/GenericPlatformFile.h"
+#include "Misc/Paths.h"
+#include "Misc/FileHelper.h"
+
+namespace SteamAudio
+{
+
+	FEnvironment::FEnvironment()
+		: EnvironmentCriticalSection()
+		, ComputeDevice(nullptr)
+		, PhononScene(nullptr)
+		, PhononEnvironment(nullptr)
+		, EnvironmentalRenderer(nullptr)
+		, ProbeManager(nullptr)
+	{
+	}
+
+	FEnvironment::~FEnvironment()
+	{
+	}
+
+	bool FEnvironment::Initialize(UWorld* World, FAudioDevice* InAudioDevice, FString* CurrentLevelName)
+	{
+		if (World == nullptr)
+		{
+			UE_LOG(LogSteamAudio, Error, TEXT("Unable to create Phonon environment: null World."));
+			return false;
+		}
+
+		if (InAudioDevice == nullptr)
+		{
+			UE_LOG(LogSteamAudio, Error, TEXT("Unable to create Phonon environment: null Audio Device."));
+			return false;
+		}
+
+		IPLerror Error = IPL_STATUS_SUCCESS;
+
+		RenderingSettings.frameSize = InAudioDevice->GetBufferLength();
+		RenderingSettings.samplingRate = InAudioDevice->GetSampleRate();
+		RenderingSettings.convolutionType = static_cast<IPLConvolutionType>(GetDefault<USteamAudioSettings>()->ConvolutionType);
+
+		IPLComputeDeviceFilter DeviceFilter;
+		DeviceFilter.fractionCUsForIRUpdate = GetDefault<USteamAudioSettings>()->FractionComputeUnitsForIRUpdate;
+		DeviceFilter.maxCUsToReserve = GetDefault<USteamAudioSettings>()->MaxComputeUnits;
+		DeviceFilter.type = IPL_COMPUTEDEVICE_GPU;
+		
+		SimulationSettings = GetDefault<USteamAudioSettings>()->GetRealtimeSimulationSettings();
+
+		// If we are using TAN or RadeonRays, attempt to create a compute device.
+		if (GetDefault<USteamAudioSettings>()->ConvolutionType == EIplConvolutionType::TRUEAUDIONEXT || 
+			GetDefault<USteamAudioSettings>()->RayTracer == EIplRayTracerType::RADEONRAYS)
+		{
+			UE_LOG(LogSteamAudio, Log, TEXT("Using TAN or Radeon Rays - creating GPU compute device..."));
+
+			Error = iplCreateComputeDevice(GlobalContext, DeviceFilter, &ComputeDevice);
+
+			// If we failed to create a compute device, fall back to Phonon scene and convolution.
+			if (Error != IPL_STATUS_SUCCESS)
+			{
+				UE_LOG(LogSteamAudio, Warning, TEXT("Failed to create GPU compute device, falling back to Phonon."));
+
+				if (SimulationSettings.sceneType == IPL_SCENETYPE_RADEONRAYS)
+				{
+					SimulationSettings.sceneType = IPL_SCENETYPE_PHONON;
+				}
+
+				if (RenderingSettings.convolutionType == IPL_CONVOLUTIONTYPE_TRUEAUDIONEXT)
+				{
+					RenderingSettings.convolutionType = IPL_CONVOLUTIONTYPE_PHONON;
+					SimulationSettings.ambisonicsOrder = GetDefault<USteamAudioSettings>()->IndirectImpulseResponseOrder;
+					SimulationSettings.irDuration = GetDefault<USteamAudioSettings>()->IndirectImpulseResponseDuration;
+					SimulationSettings.maxConvolutionSources = GetDefault<USteamAudioSettings>()->MaxSources;
+				}
+			}
+		}
+
+#if (PLATFORM_32BITS || PLATFORM_ANDROID)
+		if (SimulationSettings.sceneType != IPL_SCENETYPE_PHONON)
+		{
+			UE_LOG(LogSteamAudio, Warning, TEXT("Falling back to Phonon scene type."));
+			SimulationSettings.sceneType = IPL_SCENETYPE_PHONON;
+		}
+#endif
+
+		IPLAudioFormat EnvironmentalOutputAudioFormat;
+		EnvironmentalOutputAudioFormat.channelLayout = IPL_CHANNELLAYOUT_STEREO;
+		EnvironmentalOutputAudioFormat.channelLayoutType = IPL_CHANNELLAYOUTTYPE_AMBISONICS;
+		EnvironmentalOutputAudioFormat.channelOrder = IPL_CHANNELORDER_DEINTERLEAVED;
+		// Number of channels for an ambisonics stream is the square of the ambisonics streams' order plus one:
+		EnvironmentalOutputAudioFormat.numSpeakers = (SimulationSettings.ambisonicsOrder + 1) * (SimulationSettings.ambisonicsOrder + 1);
+		EnvironmentalOutputAudioFormat.speakerDirections = nullptr;
+		EnvironmentalOutputAudioFormat.ambisonicsOrder = SimulationSettings.ambisonicsOrder;
+		EnvironmentalOutputAudioFormat.ambisonicsNormalization = IPL_AMBISONICSNORMALIZATION_N3D;
+		EnvironmentalOutputAudioFormat.ambisonicsOrdering = IPL_AMBISONICSORDERING_ACN;
+
+		if (!LoadSceneFromDisk(World, ComputeDevice, SimulationSettings, &PhononScene, PhononSceneInfo, CurrentLevelName))
+		{
+			UE_LOG(LogSteamAudio, Warning, TEXT("Unable to create Phonon environment: failed to load scene from disk. Be sure to export the scene."));
+			return false;
+		}
+
+		Error = iplCreateProbeManager(GlobalContext, &ProbeManager);
+		LogSteamAudioStatus(Error);
+
+		TArray<AActor*> PhononProbeVolumes;
+		UGameplayStatics::GetAllActorsOfClass(World, APhononProbeVolume::StaticClass(), PhononProbeVolumes);
+
+		for (AActor* PhononProbeVolumeActor : PhononProbeVolumes)
+		{
+			APhononProbeVolume* PhononProbeVolume = Cast<APhononProbeVolume>(PhononProbeVolumeActor);
+
+			if (PhononProbeVolume->GetProbeBatchDataSize() == 0)
+			{
+				UE_LOG(LogSteamAudio, Warning, TEXT("No batch data found on probe volume. You may need to bake indirect sound."));
+				continue;
+			}
+
+			IPLhandle ProbeBatch = nullptr;
+			PhononProbeVolume->LoadProbeBatchFromDisk(&ProbeBatch);
+
+			iplAddProbeBatch(ProbeManager, ProbeBatch);
+
+			ProbeBatches.Add(ProbeBatch);
+		}
+
+		if (!LoadBakedIdentifierMapFromDisk(World, BakedIdentifierMap))
+		{
+			UE_LOG(LogSteamAudio, Warning, TEXT("Unable to load identifier map."));
+		}
+
+		Error = iplCreateEnvironment(GlobalContext, ComputeDevice, SimulationSettings, PhononScene, ProbeManager, &PhononEnvironment);
+		LogSteamAudioStatus(Error);
+
+		Error = iplCreateEnvironmentalRenderer(GlobalContext, PhononEnvironment, RenderingSettings, EnvironmentalOutputAudioFormat,
+			nullptr, nullptr, &EnvironmentalRenderer);
+		LogSteamAudioStatus(Error);
+
+		return true;
+	}
+
+	void FEnvironment::Shutdown()
+	{
+		FScopeLock EnvironmentLock(&EnvironmentCriticalSection);
+		if (ProbeManager)
+		{
+			for (IPLhandle ProbeBatch : ProbeBatches)
+			{
+				iplRemoveProbeBatch(ProbeManager, ProbeBatch);
+				iplDestroyProbeBatch(&ProbeBatch);
+			}
+
+			ProbeBatches.Empty();
+
+			iplDestroyProbeManager(&ProbeManager);
+		}
+
+		if (EnvironmentalRenderer)
+		{
+			iplDestroyEnvironmentalRenderer(&EnvironmentalRenderer);
+		}
+
+		if (PhononEnvironment)
+		{
+			iplDestroyEnvironment(&PhononEnvironment);
+		}
+
+		if (PhononScene)
+		{
+			iplDestroyScene(&PhononScene);
+		}
+
+		if (ComputeDevice)
+		{
+			iplDestroyComputeDevice(&ComputeDevice);
+		}
+	}
+
+	//==============================================================================================================================================
+	// FIdentifierMap
+	//==============================================================================================================================================
+
+	bool FIdentifierMap::ContainsKey(const FString& Key) const
+	{
+		for (const auto& Pair : IdentifierMapping)
+		{
+			if (Pair.Key == Key)
+			{
+				return true;
+			}
+		}
+
+		return false;
+	}
+
+	bool FIdentifierMap::ContainsValue(const int32 Value) const
+	{
+		for (const auto& Pair : IdentifierMapping)
+		{
+			if (Pair.Value == Value)
+			{
+				return true;
+			}
+		}
+
+		return false;
+	}
+
+	int32 FIdentifierMap::Add(const FString& Identifier)
+	{
+		int32 Value = FMath::Rand();
+
+		while (ContainsValue(Value))
+		{
+			Value = FMath::Rand();
+		}
+
+		IdentifierMapping.Add(TPair<FString, int32>(Identifier, Value));
+
+		return Value;
+	}
+
+	void FIdentifierMap::Add(const TPair<FString, int32>& Pair)
+	{
+		IdentifierMapping.Add(Pair);
+	}
+
+	int32 FIdentifierMap::Get(const FString& Identifier) const
+	{
+		for (const auto& Pair : IdentifierMapping)
+		{
+			if (Pair.Key == Identifier)
+			{
+				return Pair.Value;
+			}
+		}
+
+		return -1;
+	}
+
+	FString FIdentifierMap::ToString() const
+	{
+		FString out;
+
+		for (const auto& Pair : IdentifierMapping)
+		{
+			out += Pair.Key + ":" + FString::FromInt(Pair.Value) + ",";
+		}
+
+		out.RemoveFromEnd(",");
+
+		return out;
+	}
+
+	void FIdentifierMap::FromString(const FString& MappingString)
+	{
+		IdentifierMapping.Empty();
+
+		TArray<FString> MappingArr;
+		MappingString.ParseIntoArray(MappingArr, *FString(","));
+
+		for (const auto& PairStr : MappingArr)
+		{
+			FString Key, Value;
+			PairStr.Split(":", &Key, &Value);
+			TPair<FString, int32> Mapping(Key, FCString::Atoi(*Value));
+			IdentifierMapping.Add(Mapping);
+		}
+	}
+
+	void FIdentifierMap::Empty()
+	{
+		IdentifierMapping.Empty();
+	}
+
+	bool LoadBakedIdentifierMapFromDisk(UWorld* World, FIdentifierMap& BakedIdentifierMap)
+	{
+		FString MapName = StrippedMapName(World->GetMapName());
+		FString BakedIdentifierMapFileName = RuntimePath + MapName + ".bakedsources";
+
+		IPlatformFile& PlatformFile = FPlatformFileManager::Get().GetPlatformFile();
+		FString BakedIdentifierMapString;
+
+		if (PlatformFile.FileExists(*BakedIdentifierMapFileName))
+		{
+			FFileHelper::LoadFileToString(BakedIdentifierMapString, *BakedIdentifierMapFileName);
+		}
+		else
+		{
+			UE_LOG(LogSteamAudio, Warning, TEXT("Unable to load baked identifier map: file doesn't exist."));
+			return false;
+		}
+
+		BakedIdentifierMap.FromString(BakedIdentifierMapString);
+
+		return true;
+	}
+
+	bool SaveBakedIdentifierMapToDisk(UWorld* World, const FIdentifierMap& BakedIdentifierMap)
+	{
+		FString MapName = StrippedMapName(World->GetMapName());
+		FString BakedIdentifierMapFileName = RuntimePath + MapName + ".bakedsources";
+
+		FFileHelper::SaveStringToFile(BakedIdentifierMap.ToString(), *BakedIdentifierMapFileName);
+
+		return true;
+	}
+
+}
diff --git a/Engine/Plugins/Runtime/Steam/SteamAudio/Source/SteamAudio/Private/SteamAudioModule.cpp b/Engine/Plugins/Runtime/Steam/SteamAudio/Source/SteamAudio/Private/SteamAudioModule.cpp
--- a/Engine/Plugins/Runtime/Steam/SteamAudio/Source/SteamAudio/Private/SteamAudioModule.cpp
+++ b/Engine/Plugins/Runtime/Steam/SteamAudio/Source/SteamAudio/Private/SteamAudioModule.cpp
@@ -3,15 +3,16 @@
 //
 
 #include "SteamAudioModule.h"
+#include "PhononCommon.h"
 #include "Misc/Paths.h"
 #include "Features/IModularFeatures.h"
 
-DEFINE_LOG_CATEGORY(LogSteamAudio);
-
 IMPLEMENT_MODULE(SteamAudio::FSteamAudioModule, SteamAudio)
 
 namespace SteamAudio
 {
+	void* FSteamAudioModule::PhononDllHandle = nullptr;
+
 	static bool bModuleStartedUp = false;
 
 	FSteamAudioModule::FSteamAudioModule()
@@ -30,7 +31,34 @@ namespace SteamAudio
 
 		UE_LOG(LogSteamAudio, Log, TEXT("FSteamAudioModule Startup"));
 
-		UE_LOG(LogSteamAudio, Warning, TEXT("The Steam Audio plugin is deprecated and will be removed in a future engine release. Get Valve's version here: https://valvesoftware.github.io/steam-audio/doc/unreal/getting-started.html"))
+		UE_LOG(LogSteamAudio, Warning, TEXT("The Steam Audio plugin is deprecated and will be removed in a future engine release. See the Plugins Browser for more information."))
+
+		// Make folder paths globally available
+		BasePath = FPaths::ProjectContentDir() + "SteamAudio/";
+		RuntimePath = BasePath + "Runtime/";
+		DynamicRuntimePath = RuntimePath + "Dynamic/";
+		EditorOnlyPath = BasePath + "EditorOnly/";
+		DynamicEditorOnlyPath = EditorOnlyPath + "Dynamic/";
+
+		//Register the Steam Audio plugin factories
+		IModularFeatures::Get().RegisterModularFeature(FSpatializationPluginFactory::GetModularFeatureName(), &SpatializationPluginFactory);
+		IModularFeatures::Get().RegisterModularFeature(FReverbPluginFactory::GetModularFeatureName(), &ReverbPluginFactory);
+		IModularFeatures::Get().RegisterModularFeature(FOcclusionPluginFactory::GetModularFeatureName(), &OcclusionPluginFactory);
+
+#if PLATFORM_WINDOWS
+		if (!FSteamAudioModule::PhononDllHandle)
+		{
+#if PLATFORM_32BITS
+			FString PathToDll = FPaths::EngineDir() / TEXT("Binaries/ThirdParty/Phonon/Win32/");
+#else
+			FString PathToDll = FPaths::EngineDir() / TEXT("Binaries/ThirdParty/Phonon/Win64/");
+#endif
+
+			FString DLLToLoad = PathToDll + TEXT("phonon.dll");
+			FSteamAudioModule::PhononDllHandle = LoadDll(DLLToLoad, true);
+		}
+#endif
+		iplCreateContext(UnrealLog, nullptr, nullptr, &GlobalContext);
 	}
 
 	void FSteamAudioModule::ShutdownModule()
@@ -40,6 +68,80 @@ namespace SteamAudio
 		check(bModuleStartedUp == true);
 
 		bModuleStartedUp = false;
+
+		iplDestroyContext(&GlobalContext);
+		iplCleanup();
+
+#if PLATFORM_WINDOWS
+		if (FSteamAudioModule::PhononDllHandle)
+		{
+			FPlatformProcess::FreeDllHandle(FSteamAudioModule::PhononDllHandle);
+			FSteamAudioModule::PhononDllHandle = nullptr;
+		}
+#endif
+	}
+
+	IAudioPluginFactory* FSteamAudioModule::GetPluginFactory(EAudioPlugin PluginType)
+	{
+		switch (PluginType)
+		{
+		case EAudioPlugin::SPATIALIZATION:
+			return &SpatializationPluginFactory;
+			break;
+		case EAudioPlugin::REVERB:
+			return &ReverbPluginFactory;
+			break;
+		case EAudioPlugin::OCCLUSION:
+			return &OcclusionPluginFactory;
+			break;
+		default:
+			return nullptr;
+			break;
+		}
+	}
+
+	void FSteamAudioModule::RegisterAudioDevice(FAudioDevice* AudioDeviceHandle)
+	{
+		if (!RegisteredAudioDevices.Contains(AudioDeviceHandle))
+		{
+			TAudioPluginListenerPtr NewPhononPluginManager = TAudioPluginListenerPtr(new FPhononPluginManager());
+			AudioDeviceHandle->RegisterPluginListener(NewPhononPluginManager);
+			RegisteredAudioDevices.Add(AudioDeviceHandle);
+		}
+	}
+
+	void FSteamAudioModule::UnregisterAudioDevice(FAudioDevice* AudioDeviceHandle)
+	{
+		RegisteredAudioDevices.Remove(AudioDeviceHandle);
+	}
+
+	TAudioOcclusionPtr FOcclusionPluginFactory::CreateNewOcclusionPlugin(FAudioDevice* OwningDevice)
+	{
+		//Register audio device to the Steam Module
+		FSteamAudioModule* Module = &FModuleManager::GetModuleChecked<FSteamAudioModule>("SteamAudio");
+		if (Module != nullptr)
+		{
+			Module->RegisterAudioDevice(OwningDevice);
+		}
+		
+		return TAudioOcclusionPtr(new FPhononOcclusion());
+	}
+
+	TAudioReverbPtr FReverbPluginFactory::CreateNewReverbPlugin(FAudioDevice* OwningDevice)
+	{
+		//Register the audio device to the steam module:
+		FSteamAudioModule* Module = &FModuleManager::GetModuleChecked<FSteamAudioModule>("SteamAudio");
+		if (Module != nullptr)
+		{
+			Module->RegisterAudioDevice(OwningDevice);
+		}
+
+		return TAudioReverbPtr(new FPhononReverb());
+	}
+
+	TAudioSpatializationPtr FSpatializationPluginFactory::CreateNewSpatializationPlugin(FAudioDevice* OwningDevice)
+	{
+		return TAudioSpatializationPtr(new FPhononSpatialization());
 	}
 
 }
diff --git a/Engine/Plugins/Runtime/Steam/SteamAudio/Source/SteamAudio/Private/SteamAudioModule.h b/Engine/Plugins/Runtime/Steam/SteamAudio/Source/SteamAudio/Private/SteamAudioModule.h
--- a/Engine/Plugins/Runtime/Steam/SteamAudio/Source/SteamAudio/Private/SteamAudioModule.h
+++ b/Engine/Plugins/Runtime/Steam/SteamAudio/Source/SteamAudio/Private/SteamAudioModule.h
@@ -5,13 +5,87 @@
 #pragma once
 
 #include "ISteamAudioModule.h"
+#include "PhononSpatialization.h"
+#include "PhononOcclusion.h"
+#include "PhononReverb.h"
+#include "PhononPluginManager.h"
+#include "PhononReverbSourceSettings.h"
+#include "PhononOcclusionSourceSettings.h"
+#include "PhononSpatializationSourceSettings.h"
 #include "AudioPluginUtilities.h"
 #include "AudioDevice.h"
 
-DECLARE_LOG_CATEGORY_EXTERN(LogSteamAudio, Log, All);
-
 namespace SteamAudio
 {
+	/************************************************************************/
+	/* Plugin Factories registered with Steam Audio Module                  */
+	/************************************************************************/
+	class FSpatializationPluginFactory : public IAudioSpatializationFactory
+	{
+	public:
+		virtual FString GetDisplayName() override
+		{
+			static FString DisplayName = FString(TEXT("Steam Audio"));
+			return DisplayName;
+		}
+
+		virtual bool SupportsPlatform(const FString& PlatformName) override
+		{
+			return	PlatformName == FString(TEXT("Windows")) || PlatformName == FString(TEXT("Android"));
+		}
+
+		virtual TAudioSpatializationPtr CreateNewSpatializationPlugin(FAudioDevice* OwningDevice) override;
+
+		virtual UClass* GetCustomSpatializationSettingsClass() const override
+		{ 
+			return UPhononOcclusionSourceSettings::StaticClass(); 
+		}
+	};
+
+	class FReverbPluginFactory : public IAudioReverbFactory
+	{
+	public:
+		virtual FString GetDisplayName() override
+		{
+			static FString DisplayName = FString(TEXT("Steam Audio"));
+			return DisplayName;
+		}
+
+		virtual bool SupportsPlatform(const FString& PlatformName) override
+		{
+			return	PlatformName == FString(TEXT("Windows"));
+		}
+
+		virtual TAudioReverbPtr CreateNewReverbPlugin(FAudioDevice* OwningDevice) override;
+
+		virtual UClass* GetCustomReverbSettingsClass() const
+		{
+			return UPhononReverbSourceSettings::StaticClass();
+		}
+	};
+
+	class FOcclusionPluginFactory : public IAudioOcclusionFactory
+	{
+	public:
+		virtual FString GetDisplayName() override
+		{
+			static FString DisplayName = FString(TEXT("Steam Audio"));
+			return DisplayName;
+		}
+
+		virtual bool SupportsPlatform(const FString& PlatformName) override
+		{
+			return	PlatformName == FString(TEXT("Windows")) || PlatformName == FString(TEXT("Android"));
+		}
+
+		virtual TAudioOcclusionPtr CreateNewOcclusionPlugin(FAudioDevice* OwningDevice) override;
+
+		virtual UClass* GetCustomOcclusionSettingsClass() const
+		{
+			return UPhononOcclusionSourceSettings::StaticClass();
+		}
+	};
+
 	/************************************************************************/
 	/* FSteamAudioModule                                                    */
 	/* Module interface for Steam Audio. Registers the Plugin Factories     */
@@ -26,5 +100,40 @@ namespace SteamAudio
 		virtual void StartupModule() override;
 		virtual void ShutdownModule() override;
 		//~ End IAudioPlugin
+
+		/** 
+		*  Returns pointer to each type of plugin factory with Steam Audio.
+		*  returns nullptr if given an invalid parameter for PluginType.
+		*/
+		IAudioPluginFactory* GetPluginFactory(EAudioPlugin PluginType);
+
+		/*
+		 * Registers a given audio device to the Steam Audio Device.
+		 * Every audio device running Steam Audio's reverb or occlusion manager
+		 * requires a PhononPluginManager registered to it.
+		 */
+		void RegisterAudioDevice(FAudioDevice* AudioDeviceHandle);
+
+		/*
+		* Unregisters an audio device from the module.
+		*/
+		void UnregisterAudioDevice(FAudioDevice* AudioDeviceHandle);
+
+	private:
+		static void* PhononDllHandle;
+		static void* TANDllHandle;
+		static void* TANUtilsDllHandle;
+		static void* EmbreeDllHandle;
+		static void* RadeonRaysDllHandle;
+		static void* TbbDllHandle;
+		static void* TbbMallocDllHandle;
+
+
+		TArray<FAudioDevice*> RegisteredAudioDevices;
+
+		//Factories:
+		FSpatializationPluginFactory SpatializationPluginFactory;
+		FReverbPluginFactory ReverbPluginFactory;
+		FOcclusionPluginFactory OcclusionPluginFactory;
 	};
 }
\ No newline at end of file
diff --git a/Engine/Plugins/Runtime/Steam/SteamAudio/Source/SteamAudio/Private/SteamAudioSettings.cpp b/Engine/Plugins/Runtime/Steam/SteamAudio/Source/SteamAudio/Private/SteamAudioSettings.cpp
new file mode 100644
--- /dev/null
+++ b/Engine/Plugins/Runtime/Steam/SteamAudio/Source/SteamAudio/Private/SteamAudioSettings.cpp
@@ -0,0 +1,370 @@
+//
+// Copyright (C) Valve Corporation. All rights reserved.
+//
+
+#include "SteamAudioSettings.h"
+
+USteamAudioSettings::USteamAudioSettings()
+	: AudioEngine(EIplAudioEngine::UNREAL)
+	, RayTracer(EIplRayTracerType::PHONON)
+	, ConvolutionType(EIplConvolutionType::PHONON)
+	, ExportBSPGeometry(true)
+	, ExportLandscapeGeometry(true)
+	, StaticMeshMaterialPreset(EPhononMaterial::GENERIC)
+	, BSPMaterialPreset(EPhononMaterial::GENERIC)
+	, LandscapeMaterialPreset(EPhononMaterial::GENERIC)
+	, OcclusionSampleCount(128)
+	, ListenerReverbSimulationType(EIplSimulationType::DISABLED)
+	, ListenerReverbContribution(1.0f)
+	, IndirectImpulseResponseOrder(1)
+	, IndirectImpulseResponseDuration(1.0f)
+	, IndirectSpatializationMethod(EIplSpatializationMethod::PANNING)
+	, IrradianceMinDistance(1.0f)
+	, MaxSources(32)
+	, RealtimeQualityPreset(EQualitySettings::LOW)
+	, RealTimeCPUCoresPercentage(5)
+	, RealtimeBounces(2)
+	, RealtimeRays(8192)
+	, RealtimeSecondaryRays(512)
+	, BakedQualityPreset(EQualitySettings::LOW)
+	, BakingCPUCoresPercentage(5)
+	, BakedBounces(128)
+	, BakedRays(16384)
+	, BakedSecondaryRays(2048)
+	, MaxComputeUnits(8)
+	, FractionComputeUnitsForIRUpdate(0.5f)
+	, TANIndirectImpulseResponseOrder(1)
+	, TANIndirectImpulseResponseDuration(1.0f)
+	, TANMaxSources(32)
+	, RadeonRaysBakingBatchSize(1)
+{
+	OutputSubmix = FString(TEXT("/SteamAudio/PhononSubmixDefault.PhononSubmixDefault"));
+
+	auto MaterialPreset = SteamAudio::MaterialPresets[StaticMeshMaterialPreset];
+	StaticMeshLowFreqAbsorption = MaterialPreset.lowFreqAbsorption;
+	StaticMeshMidFreqAbsorption = MaterialPreset.midFreqAbsorption;
+	StaticMeshHighFreqAbsorption = MaterialPreset.highFreqAbsorption;
+	StaticMeshLowFreqTransmission = MaterialPreset.lowFreqTransmission;
+	StaticMeshMidFreqTransmission = MaterialPreset.midFreqTransmission;
+	StaticMeshHighFreqTransmission = MaterialPreset.highFreqTransmission;
+	StaticMeshScattering = MaterialPreset.scattering;
+
+	MaterialPreset = SteamAudio::MaterialPresets[BSPMaterialPreset];
+	BSPLowFreqAbsorption = MaterialPreset.lowFreqAbsorption;
+	BSPMidFreqAbsorption = MaterialPreset.midFreqAbsorption;
+	BSPHighFreqAbsorption = MaterialPreset.highFreqAbsorption;
+	BSPLowFreqTransmission = MaterialPreset.lowFreqTransmission;
+	BSPMidFreqTransmission = MaterialPreset.midFreqTransmission;
+	BSPHighFreqTransmission = MaterialPreset.highFreqTransmission;
+	BSPScattering = MaterialPreset.scattering;
+
+	MaterialPreset = SteamAudio::MaterialPresets[LandscapeMaterialPreset];
+	LandscapeLowFreqAbsorption = MaterialPreset.lowFreqAbsorption;
+	LandscapeMidFreqAbsorption = MaterialPreset.midFreqAbsorption;
+	LandscapeHighFreqAbsorption = MaterialPreset.highFreqAbsorption;
+	LandscapeLowFreqTransmission = MaterialPreset.lowFreqTransmission;
+	LandscapeMidFreqTransmission = MaterialPreset.midFreqTransmission;
+	LandscapeHighFreqTransmission = MaterialPreset.highFreqTransmission;
+	LandscapeScattering = MaterialPreset.scattering;
+}
+
+IPLMaterial USteamAudioSettings::GetDefaultStaticMeshMaterial() const
+{
+	IPLMaterial DAM;
+	DAM.lowFreqAbsorption = StaticMeshLowFreqAbsorption;
+	DAM.midFreqAbsorption = StaticMeshMidFreqAbsorption;
+	DAM.highFreqAbsorption = StaticMeshHighFreqAbsorption;
+	DAM.lowFreqTransmission = StaticMeshLowFreqTransmission;
+	DAM.midFreqTransmission = StaticMeshMidFreqTransmission;
+	DAM.highFreqTransmission = StaticMeshHighFreqTransmission;
+	DAM.scattering = StaticMeshScattering;
+
+	return DAM;
+}
+
+IPLMaterial USteamAudioSettings::GetDefaultBSPMaterial() const
+{
+	IPLMaterial DAM;
+	DAM.lowFreqAbsorption = BSPLowFreqAbsorption;
+	DAM.midFreqAbsorption = BSPMidFreqAbsorption;
+	DAM.highFreqAbsorption = BSPHighFreqAbsorption;
+	DAM.lowFreqTransmission = BSPLowFreqTransmission;
+	DAM.midFreqTransmission = BSPMidFreqTransmission;
+	DAM.highFreqTransmission = BSPHighFreqTransmission;
+	DAM.scattering = BSPScattering;
+
+	return DAM;
+}
+
+IPLMaterial USteamAudioSettings::GetDefaultLandscapeMaterial() const
+{
+	IPLMaterial DAM;
+	DAM.lowFreqAbsorption = LandscapeLowFreqAbsorption;
+	DAM.midFreqAbsorption = LandscapeMidFreqAbsorption;
+	DAM.highFreqAbsorption = LandscapeHighFreqAbsorption;
+	DAM.lowFreqTransmission = LandscapeLowFreqTransmission;
+	DAM.midFreqTransmission = LandscapeMidFreqTransmission;
+	DAM.highFreqTransmission = LandscapeHighFreqTransmission;
+	DAM.scattering = LandscapeScattering;
+
+	return DAM;
+}
+
+int32 USteamAudioSettings::GetIndirectImpulseResponseOrder() const
+{
+	return ConvolutionType == EIplConvolutionType::TRUEAUDIONEXT ? TANIndirectImpulseResponseOrder : IndirectImpulseResponseOrder;
+}
+
+float USteamAudioSettings::GetIndirectImpulseResponseDuration() const
+{
+	return ConvolutionType == EIplConvolutionType::TRUEAUDIONEXT ? TANIndirectImpulseResponseDuration : IndirectImpulseResponseDuration;
+}
+
+uint32 USteamAudioSettings::GetMaxSources() const
+{
+	return ConvolutionType == EIplConvolutionType::TRUEAUDIONEXT ? TANMaxSources : MaxSources;
+}
+
+int32 USteamAudioSettings::GetBakingBatchSize() const
+{
+	return RayTracer == EIplRayTracerType::RADEONRAYS ? RadeonRaysBakingBatchSize : 1;
+}
+
+float USteamAudioSettings::GetFractionComputeUnitsForIRUpdate() const
+{
+	return ConvolutionType == EIplConvolutionType::TRUEAUDIONEXT ? FractionComputeUnitsForIRUpdate : 1.0f;
+}
+
+IPLSimulationSettings USteamAudioSettings::GetRealtimeSimulationSettings() const
+{
+	IPLSimulationSettings SimulationSettings;
+	SimulationSettings.numBounces = RealtimeBounces;
+	SimulationSettings.numDiffuseSamples = RealtimeSecondaryRays;
+	SimulationSettings.numRays = RealtimeRays;
+	SimulationSettings.numOcclusionSamples = OcclusionSampleCount;
+	SimulationSettings.maxConvolutionSources = GetMaxSources();
+	SimulationSettings.ambisonicsOrder = GetIndirectImpulseResponseOrder();
+	SimulationSettings.irDuration = GetIndirectImpulseResponseDuration();
+	SimulationSettings.irradianceMinDistance = IrradianceMinDistance;
+	SimulationSettings.sceneType = static_cast<IPLSceneType>(RayTracer);
+	SimulationSettings.numThreads = FMath::Clamp<int32>(FPlatformMisc::NumberOfCoresIncludingHyperthreads() * (static_cast<float>(RealTimeCPUCoresPercentage) / 100.0f), 1, FPlatformMisc::NumberOfCoresIncludingHyperthreads());
+	SimulationSettings.bakingBatchSize = GetBakingBatchSize();
+
+	return SimulationSettings;
+}
+
+IPLSimulationSettings USteamAudioSettings::GetBakedSimulationSettings() const
+{
+	bool bUseTANOverrides = ConvolutionType == EIplConvolutionType::TRUEAUDIONEXT;
+
+	IPLSimulationSettings SimulationSettings;
+	SimulationSettings.numBounces = BakedBounces;
+	SimulationSettings.numDiffuseSamples = BakedSecondaryRays;
+	SimulationSettings.numRays = BakedRays;
+	SimulationSettings.numOcclusionSamples = OcclusionSampleCount;
+	SimulationSettings.maxConvolutionSources = GetMaxSources();
+	SimulationSettings.ambisonicsOrder = GetIndirectImpulseResponseOrder();
+	SimulationSettings.irDuration = GetIndirectImpulseResponseDuration();
+	SimulationSettings.irradianceMinDistance = IrradianceMinDistance;
+	SimulationSettings.sceneType = static_cast<IPLSceneType>(RayTracer);
+	SimulationSettings.numThreads = FMath::Clamp<int32>(FPlatformMisc::NumberOfCoresIncludingHyperthreads() * (static_cast<float>(BakingCPUCoresPercentage) / 100.0f), 1, FPlatformMisc::NumberOfCoresIncludingHyperthreads());
+	SimulationSettings.bakingBatchSize = GetBakingBatchSize();
+
+	return SimulationSettings;
+}
+
+#if WITH_EDITOR
+void USteamAudioSettings::PostEditChangeProperty(FPropertyChangedEvent& PropertyChangedEvent)
+{
+	Super::PostEditChangeProperty(PropertyChangedEvent);
+
+	FName PropertyName = (PropertyChangedEvent.Property != nullptr) ? PropertyChangedEvent.Property->GetFName() : NAME_None;
+
+	auto MaterialPreset = SteamAudio::MaterialPresets[StaticMeshMaterialPreset];
+	auto RealtimeSimulationQuality = SteamAudio::RealtimeSimulationQualityPresets[RealtimeQualityPreset];
+	auto BakedSimulationQuality = SteamAudio::BakedSimulationQualityPresets[BakedQualityPreset];
+
+	if ((PropertyName == GET_MEMBER_NAME_CHECKED(USteamAudioSettings, StaticMeshMaterialPreset)))
+	{
+		StaticMeshLowFreqAbsorption = MaterialPreset.lowFreqAbsorption;
+		StaticMeshMidFreqAbsorption = MaterialPreset.midFreqAbsorption;
+		StaticMeshHighFreqAbsorption = MaterialPreset.highFreqAbsorption;
+		StaticMeshLowFreqTransmission = MaterialPreset.lowFreqTransmission;
+		StaticMeshMidFreqTransmission = MaterialPreset.midFreqTransmission;
+		StaticMeshHighFreqTransmission = MaterialPreset.highFreqTransmission;
+		StaticMeshScattering = MaterialPreset.scattering;
+
+		for (TFieldIterator<FProperty> PropIt(GetClass()); PropIt; ++PropIt) 
+		{
+			FProperty* Property = *PropIt;
+			if (Property->GetFName() == GET_MEMBER_NAME_CHECKED(USteamAudioSettings, StaticMeshLowFreqAbsorption) ||
+				Property->GetFName() == GET_MEMBER_NAME_CHECKED(USteamAudioSettings, StaticMeshMidFreqAbsorption) ||
+				Property->GetFName() == GET_MEMBER_NAME_CHECKED(USteamAudioSettings, StaticMeshHighFreqAbsorption) ||
+				Property->GetFName() == GET_MEMBER_NAME_CHECKED(USteamAudioSettings, StaticMeshLowFreqTransmission) ||
+				Property->GetFName() == GET_MEMBER_NAME_CHECKED(USteamAudioSettings, StaticMeshMidFreqTransmission) ||
+				Property->GetFName() == GET_MEMBER_NAME_CHECKED(USteamAudioSettings, StaticMeshHighFreqTransmission) ||
+				Property->GetFName() == GET_MEMBER_NAME_CHECKED(USteamAudioSettings, StaticMeshScattering))
+			{
+				UpdateSinglePropertyInConfigFile(Property, GetDefaultConfigFilename());
+			}
+		}
+	}
+	else if ((PropertyName == GET_MEMBER_NAME_CHECKED(USteamAudioSettings, BSPMaterialPreset)))
+	{
+		MaterialPreset = SteamAudio::MaterialPresets[BSPMaterialPreset];
+		BSPLowFreqAbsorption = MaterialPreset.lowFreqAbsorption;
+		BSPMidFreqAbsorption = MaterialPreset.midFreqAbsorption;
+		BSPHighFreqAbsorption = MaterialPreset.highFreqAbsorption;
+		BSPLowFreqTransmission = MaterialPreset.lowFreqTransmission;
+		BSPMidFreqTransmission = MaterialPreset.midFreqTransmission;
+		BSPHighFreqTransmission = MaterialPreset.highFreqTransmission;
+		BSPScattering = MaterialPreset.scattering;
+
+		for (TFieldIterator<FProperty> PropIt(GetClass()); PropIt; ++PropIt)
+		{
+			FProperty* Property = *PropIt;
+			if (Property->GetFName() == GET_MEMBER_NAME_CHECKED(USteamAudioSettings, BSPLowFreqAbsorption) ||
+				Property->GetFName() == GET_MEMBER_NAME_CHECKED(USteamAudioSettings, BSPMidFreqAbsorption) ||
+				Property->GetFName() == GET_MEMBER_NAME_CHECKED(USteamAudioSettings, BSPHighFreqAbsorption) ||
+				Property->GetFName() == GET_MEMBER_NAME_CHECKED(USteamAudioSettings, BSPLowFreqTransmission) ||
+				Property->GetFName() == GET_MEMBER_NAME_CHECKED(USteamAudioSettings, BSPMidFreqTransmission) ||
+				Property->GetFName() == GET_MEMBER_NAME_CHECKED(USteamAudioSettings, BSPHighFreqTransmission) ||
+				Property->GetFName() == GET_MEMBER_NAME_CHECKED(USteamAudioSettings, BSPScattering))
+			{
+				UpdateSinglePropertyInConfigFile(Property, GetDefaultConfigFilename());
+			}
+		}
+	}
+	else if ((PropertyName == GET_MEMBER_NAME_CHECKED(USteamAudioSettings, LandscapeMaterialPreset)))
+	{
+		MaterialPreset = SteamAudio::MaterialPresets[LandscapeMaterialPreset];
+		LandscapeLowFreqAbsorption = MaterialPreset.lowFreqAbsorption;
+		LandscapeMidFreqAbsorption = MaterialPreset.midFreqAbsorption;
+		LandscapeHighFreqAbsorption = MaterialPreset.highFreqAbsorption;
+		LandscapeLowFreqTransmission = MaterialPreset.lowFreqTransmission;
+		LandscapeMidFreqTransmission = MaterialPreset.midFreqTransmission;
+		LandscapeHighFreqTransmission = MaterialPreset.highFreqTransmission;
+		LandscapeScattering = MaterialPreset.scattering;
+
+		for (TFieldIterator<FProperty> PropIt(GetClass()); PropIt; ++PropIt)
+		{
+			FProperty* Property = *PropIt;
+			if (Property->GetFName() == GET_MEMBER_NAME_CHECKED(USteamAudioSettings, LandscapeLowFreqAbsorption) ||
+				Property->GetFName() == GET_MEMBER_NAME_CHECKED(USteamAudioSettings, LandscapeMidFreqAbsorption) ||
+				Property->GetFName() == GET_MEMBER_NAME_CHECKED(USteamAudioSettings, LandscapeHighFreqAbsorption) ||
+				Property->GetFName() == GET_MEMBER_NAME_CHECKED(USteamAudioSettings, LandscapeLowFreqTransmission) ||
+				Property->GetFName() == GET_MEMBER_NAME_CHECKED(USteamAudioSettings, LandscapeMidFreqTransmission) ||
+				Property->GetFName() == GET_MEMBER_NAME_CHECKED(USteamAudioSettings, LandscapeHighFreqTransmission) ||
+				Property->GetFName() == GET_MEMBER_NAME_CHECKED(USteamAudioSettings, LandscapeScattering))
+			{
+				UpdateSinglePropertyInConfigFile(Property, GetDefaultConfigFilename());
+			}
+		}
+	}
+	else if ((PropertyName == GET_MEMBER_NAME_CHECKED(USteamAudioSettings, RealtimeQualityPreset)))
+	{
+		RealtimeBounces = RealtimeSimulationQuality.Bounces;
+		RealtimeRays = RealtimeSimulationQuality.Rays;
+		RealtimeSecondaryRays = RealtimeSimulationQuality.SecondaryRays;
+
+		for (TFieldIterator<FProperty> PropIt(GetClass()); PropIt; ++PropIt)
+		{
+			FProperty* Property = *PropIt;
+			if (Property->GetFName() == GET_MEMBER_NAME_CHECKED(USteamAudioSettings, RealtimeBounces) ||
+				Property->GetFName() == GET_MEMBER_NAME_CHECKED(USteamAudioSettings, RealtimeRays) ||
+				Property->GetFName() == GET_MEMBER_NAME_CHECKED(USteamAudioSettings, RealtimeSecondaryRays))
+			{
+				UpdateSinglePropertyInConfigFile(Property, GetDefaultConfigFilename());
+			}
+		}
+	}
+	else if ((PropertyName == GET_MEMBER_NAME_CHECKED(USteamAudioSettings, BakedQualityPreset)))
+	{
+		BakedBounces = BakedSimulationQuality.Bounces;
+		BakedRays = BakedSimulationQuality.Rays;
+		BakedSecondaryRays = BakedSimulationQuality.SecondaryRays;
+
+		for (TFieldIterator<FProperty> PropIt(GetClass()); PropIt; ++PropIt)
+		{
+			FProperty* Property = *PropIt;
+			if (Property->GetFName() == GET_MEMBER_NAME_CHECKED(USteamAudioSettings, BakedBounces) ||
+				Property->GetFName() == GET_MEMBER_NAME_CHECKED(USteamAudioSettings, BakedRays) ||
+				Property->GetFName() == GET_MEMBER_NAME_CHECKED(USteamAudioSettings, BakedSecondaryRays))
+			{
+				UpdateSinglePropertyInConfigFile(Property, GetDefaultConfigFilename());
+			}
+		}
+	}
+}
+
+bool USteamAudioSettings::CanEditChange(const FProperty* InProperty) const
+{
+	const bool ParentVal = Super::CanEditChange(InProperty);
+
+	if (InProperty->GetFName() == GET_MEMBER_NAME_CHECKED(USteamAudioSettings, FractionComputeUnitsForIRUpdate) ||
+		InProperty->GetFName() == GET_MEMBER_NAME_CHECKED(USteamAudioSettings, TANIndirectImpulseResponseOrder) ||
+		InProperty->GetFName() == GET_MEMBER_NAME_CHECKED(USteamAudioSettings, TANIndirectImpulseResponseDuration) ||
+		InProperty->GetFName() == GET_MEMBER_NAME_CHECKED(USteamAudioSettings, TANMaxSources))
+	{
+		return ParentVal && ConvolutionType == EIplConvolutionType::TRUEAUDIONEXT;
+	}
+	else if (InProperty->GetFName() == GET_MEMBER_NAME_CHECKED(USteamAudioSettings, MaxComputeUnits))
+	{
+		return ParentVal && (ConvolutionType == EIplConvolutionType::TRUEAUDIONEXT || RayTracer == EIplRayTracerType::RADEONRAYS);
+	}
+	else if (InProperty->GetFName() == GET_MEMBER_NAME_CHECKED(USteamAudioSettings, RadeonRaysBakingBatchSize))
+	{
+		return ParentVal && RayTracer == EIplRayTracerType::RADEONRAYS;
+	}
+	else if (InProperty->GetFName() == GET_MEMBER_NAME_CHECKED(USteamAudioSettings, StaticMeshLowFreqAbsorption) ||
+			 InProperty->GetFName() == GET_MEMBER_NAME_CHECKED(USteamAudioSettings, StaticMeshMidFreqAbsorption) ||
+			 InProperty->GetFName() == GET_MEMBER_NAME_CHECKED(USteamAudioSettings, StaticMeshHighFreqAbsorption) ||
+			 InProperty->GetFName() == GET_MEMBER_NAME_CHECKED(USteamAudioSettings, StaticMeshLowFreqTransmission) ||
+			 InProperty->GetFName() == GET_MEMBER_NAME_CHECKED(USteamAudioSettings, StaticMeshMidFreqTransmission) ||
+			 InProperty->GetFName() == GET_MEMBER_NAME_CHECKED(USteamAudioSettings, StaticMeshHighFreqTransmission) ||
+			 InProperty->GetFName() == GET_MEMBER_NAME_CHECKED(USteamAudioSettings, StaticMeshScattering))
+	{
+		return ParentVal && (StaticMeshMaterialPreset == EPhononMaterial::CUSTOM);
+	}
+	else if (InProperty->GetFName() == GET_MEMBER_NAME_CHECKED(USteamAudioSettings, BSPLowFreqAbsorption) ||
+			 InProperty->GetFName() == GET_MEMBER_NAME_CHECKED(USteamAudioSettings, BSPMidFreqAbsorption) ||
+			 InProperty->GetFName() == GET_MEMBER_NAME_CHECKED(USteamAudioSettings, BSPHighFreqAbsorption) ||
+			 InProperty->GetFName() == GET_MEMBER_NAME_CHECKED(USteamAudioSettings, BSPLowFreqTransmission) ||
+			 InProperty->GetFName() == GET_MEMBER_NAME_CHECKED(USteamAudioSettings, BSPMidFreqTransmission) ||
+			 InProperty->GetFName() == GET_MEMBER_NAME_CHECKED(USteamAudioSettings, BSPHighFreqTransmission) ||
+			 InProperty->GetFName() == GET_MEMBER_NAME_CHECKED(USteamAudioSettings, BSPScattering))
+	{
+		return ParentVal && (BSPMaterialPreset == EPhononMaterial::CUSTOM);
+	}
+	else if (InProperty->GetFName() == GET_MEMBER_NAME_CHECKED(USteamAudioSettings, LandscapeLowFreqAbsorption) ||
+			 InProperty->GetFName() == GET_MEMBER_NAME_CHECKED(USteamAudioSettings, LandscapeMidFreqAbsorption) ||
+			 InProperty->GetFName() == GET_MEMBER_NAME_CHECKED(USteamAudioSettings, LandscapeHighFreqAbsorption) ||
+			 InProperty->GetFName() == GET_MEMBER_NAME_CHECKED(USteamAudioSettings, LandscapeLowFreqTransmission) ||
+			 InProperty->GetFName() == GET_MEMBER_NAME_CHECKED(USteamAudioSettings, LandscapeMidFreqTransmission) ||
+			 InProperty->GetFName() == GET_MEMBER_NAME_CHECKED(USteamAudioSettings, LandscapeHighFreqTransmission) ||
+			 InProperty->GetFName() == GET_MEMBER_NAME_CHECKED(USteamAudioSettings, LandscapeScattering))
+	{
+		return ParentVal && (LandscapeMaterialPreset == EPhononMaterial::CUSTOM);
+	}
+	else if (InProperty->GetFName() == GET_MEMBER_NAME_CHECKED(USteamAudioSettings, RealtimeBounces) ||
+			 InProperty->GetFName() == GET_MEMBER_NAME_CHECKED(USteamAudioSettings, RealtimeRays) ||
+			 InProperty->GetFName() == GET_MEMBER_NAME_CHECKED(USteamAudioSettings, RealtimeSecondaryRays))
+	{
+		return ParentVal && (RealtimeQualityPreset == EQualitySettings::CUSTOM);
+	}
+	else if (InProperty->GetFName() == GET_MEMBER_NAME_CHECKED(USteamAudioSettings, BakedBounces) ||
+			 InProperty->GetFName() == GET_MEMBER_NAME_CHECKED(USteamAudioSettings, BakedRays) ||
+			 InProperty->GetFName() == GET_MEMBER_NAME_CHECKED(USteamAudioSettings, BakedSecondaryRays))
+	{
+		return ParentVal && (BakedQualityPreset == EQualitySettings::CUSTOM);
+	}
+	else if (InProperty->GetFName() == GET_MEMBER_NAME_CHECKED(USteamAudioSettings, ListenerReverbContribution))
+	{
+		return ParentVal && ListenerReverbSimulationType != EIplSimulationType::DISABLED;
+	}
+	else
+	{
+		return ParentVal;
+	}
+}
+#endif
\ No newline at end of file
diff --git a/Engine/Plugins/Runtime/Steam/SteamAudio/Source/SteamAudio/Public/PhononOcclusionSourceSettings.h b/Engine/Plugins/Runtime/Steam/SteamAudio/Source/SteamAudio/Public/PhononOcclusionSourceSettings.h
new file mode 100644
--- /dev/null
+++ b/Engine/Plugins/Runtime/Steam/SteamAudio/Source/SteamAudio/Public/PhononOcclusionSourceSettings.h
@@ -0,0 +1,38 @@
+//
+// Copyright (C) Valve Corporation. All rights reserved.
+//
+
+#pragma once
+
+#include "IAudioExtensionPlugin.h"
+#include "PhononCommon.h"
+#include "PhononOcclusionSourceSettings.generated.h"
+
+UCLASS()
+class STEAMAUDIO_API UPhononOcclusionSourceSettings : public UOcclusionPluginSourceSettingsBase
+{
+	GENERATED_BODY()
+
+public:
+	UPhononOcclusionSourceSettings();
+
+#if WITH_EDITOR
+	virtual bool CanEditChange(const FProperty* InProperty) const override;
+#endif
+
+	UPROPERTY(EditAnywhere, BlueprintReadWrite, Category = OcclusionSettings)
+	EIplDirectOcclusionMode DirectOcclusionMode;
+
+	UPROPERTY(EditAnywhere, BlueprintReadWrite, Category = OcclusionSettings)
+	EIplDirectOcclusionMethod DirectOcclusionMethod;
+
+	UPROPERTY(EditAnywhere, BlueprintReadWrite, Category = OcclusionSettings)
+	float DirectOcclusionSourceRadius;
+
+	UPROPERTY(EditAnywhere, BlueprintReadWrite, Category = OcclusionSettings, meta = (DisplayName = "Physics-based Attenuation"))
+	bool DirectAttenuation;
+
+	UPROPERTY(EditAnywhere, BlueprintReadWrite, Category = OcclusionSettings)
+	bool AirAbsorption;
+};
+
diff --git a/Engine/Plugins/Runtime/Steam/SteamAudio/Source/SteamAudio/Public/PhononReverbSourceSettings.h b/Engine/Plugins/Runtime/Steam/SteamAudio/Source/SteamAudio/Public/PhononReverbSourceSettings.h
new file mode 100644
--- /dev/null
+++ b/Engine/Plugins/Runtime/Steam/SteamAudio/Source/SteamAudio/Public/PhononReverbSourceSettings.h
@@ -0,0 +1,26 @@
+//
+// Copyright (C) Valve Corporation. All rights reserved.
+//
+
+#pragma once
+
+#include "IAudioExtensionPlugin.h"
+#include "PhononCommon.h"
+#include "PhononReverbSourceSettings.generated.h"
+
+UCLASS()
+class STEAMAUDIO_API UPhononReverbSourceSettings : public UReverbPluginSourceSettingsBase
+{
+	GENERATED_BODY()
+
+public:
+	UPhononReverbSourceSettings();
+
+	// How to simulate source-centric reverb.
+	UPROPERTY(EditAnywhere, BlueprintReadWrite, Category = ReverbSettings, meta = (DisplayName = "Source-Centric Reverb Simulation"))
+	EIplSimulationType SourceReverbSimulationType;
+
+	// Scale factor to make the indirect contribution louder or softer.
+	UPROPERTY(GlobalConfig, EditAnywhere, Category = ReverbSettings, meta = (DisplayName = "Source-Centric Reverb Contribution", ClampMin = "0.0", ClampMax = "10.0", UIMin = "0.0", UIMax = "10.0"))
+	float SourceReverbContribution;
+};
diff --git a/Engine/Plugins/Runtime/Steam/SteamAudio/Source/SteamAudio/Public/PhononSpatializationSourceSettings.h b/Engine/Plugins/Runtime/Steam/SteamAudio/Source/SteamAudio/Public/PhononSpatializationSourceSettings.h
new file mode 100644
--- /dev/null
+++ b/Engine/Plugins/Runtime/Steam/SteamAudio/Source/SteamAudio/Public/PhononSpatializationSourceSettings.h
@@ -0,0 +1,29 @@
+//
+// Copyright (C) Valve Corporation. All rights reserved.
+//
+
+#pragma once
+
+#include "IAudioExtensionPlugin.h"
+#include "PhononCommon.h"
+#include "PhononSpatializationSourceSettings.generated.h"
+
+UCLASS()
+class STEAMAUDIO_API UPhononSpatializationSourceSettings : public USpatializationPluginSourceSettingsBase
+{
+	GENERATED_BODY()
+
+public:
+	UPhononSpatializationSourceSettings();
+		
+#if WITH_EDITOR
+	virtual bool CanEditChange(const FProperty* InProperty) const override;
+#endif
+
+	UPROPERTY(GlobalConfig, EditAnywhere, Category = SpatializationSettings)
+	EIplSpatializationMethod SpatializationMethod;
+
+	UPROPERTY(GlobalConfig, EditAnywhere, Category = SpatializationSettings, meta = (DisplayName = "HRTF Interpolation Method"))
+	EIplHrtfInterpolationMethod HrtfInterpolationMethod;
+};
+
