[Backout] - CL25430596

diff --git a/Engine/Plugins/Runtime/Steam/SteamAudio/Source/SteamAudioEditor/Private/BakeIndirectWindow.cpp b/Engine/Plugins/Runtime/Steam/SteamAudio/Source/SteamAudioEditor/Private/BakeIndirectWindow.cpp
new file mode 100644
--- /dev/null
+++ b/Engine/Plugins/Runtime/Steam/SteamAudio/Source/SteamAudioEditor/Private/BakeIndirectWindow.cpp
@@ -0,0 +1,260 @@
+//
+// Copyright (C) Valve Corporation. All rights reserved.
+//
+
+#include "BakeIndirectWindow.h"
+
+#include "PhononCommon.h"
+#include "PhononSourceComponent.h"
+#include "PhononProbeVolume.h"
+#include "TickableNotification.h"
+#include "IndirectBaker.h"
+
+#include "Async/Async.h"
+#include "DetailLayoutBuilder.h"
+#include "EngineUtils.h"
+#include "Framework/Docking/TabManager.h"
+#include "Framework/Notifications/NotificationManager.h"
+#include "Kismet/GameplayStatics.h"
+#include "LevelEditor.h"
+#include "LevelEditorActions.h"
+#include "LevelEditorViewport.h"
+#include "Widgets/Docking/SDockTab.h"
+#include "Widgets/Input/SButton.h"
+#include "Widgets/Input/SComboBox.h"
+#include "Widgets/Input/SSpinBox.h"
+#include "Widgets/Layout/SExpandableArea.h"
+#include "Widgets/Layout/SGridPanel.h"
+#include "Widgets/Layout/SScrollBox.h"
+#include "Widgets/Notifications/SNotificationList.h"
+#include "Widgets/Views/SListView.h"
+
+namespace SteamAudio
+{
+	//==============================================================================================================================================
+	// FBakeIndirectWindow
+	//==============================================================================================================================================
+
+	static const FName BakeIndirectTabName("BakeIndirectTab");
+
+	FBakeIndirectWindow::FBakeIndirectWindow()
+	{
+		FGlobalTabmanager::Get()->RegisterTabSpawner(BakeIndirectTabName, FOnSpawnTab::CreateRaw(this, &FBakeIndirectWindow::SpawnTab))
+			.SetDisplayName(FText::FromString(TEXT("Bake Indirect Sound")));
+	}
+
+	FBakeIndirectWindow::~FBakeIndirectWindow()
+	{
+		FGlobalTabmanager::Get()->UnregisterTabSpawner(BakeIndirectTabName);
+	}
+
+	TSharedRef<SDockTab> FBakeIndirectWindow::SpawnTab(const FSpawnTabArgs& TabSpawnArgs)
+	{
+		RefreshBakedSources();
+
+		BakedSourcesListView = SNew(SListView<TSharedPtr<FBakedSource>>)
+			.ListItemsSource(&BakedSources)
+			.ScrollbarVisibility(EVisibility::Visible)
+			.OnGenerateRow(this, &FBakeIndirectWindow::OnGenerateBakedSourceRow)
+			.HeaderRow(
+				SNew(SHeaderRow)
+				+ SHeaderRow::Column("Unique Identifier")
+				.DefaultLabel(FText::FromString(TEXT("Unique Identifier")))
+				.FillWidth(0.5f)
+				+ SHeaderRow::Column("Baked Data Size")
+				.DefaultLabel(FText::FromString(TEXT("Baked Data Size")))
+				.FillWidth(0.5f)
+			)
+			.SelectionMode(ESelectionMode::Multi);
+
+		TSharedRef<SDockTab> SpawnedTab = SNew(SDockTab)[
+			SNew(SVerticalBox)
+			+SVerticalBox::Slot()
+			[
+				SNew(SBorder)
+				.BorderImage(FAppStyle::GetBrush("ToolPanel.GroupBorder"))
+				[
+					BakedSourcesListView.ToSharedRef()
+				]
+			]
+			+ SVerticalBox::Slot()
+			.AutoHeight()
+			.HAlign(HAlign_Right)
+			.Padding(2)
+			[
+				SNew(SHorizontalBox)
+				+ SHorizontalBox::Slot()
+				.AutoWidth()
+				[
+					SNew(SButton)
+					.ContentPadding(3)
+					.VAlign(VAlign_Center)
+					.HAlign(HAlign_Center)
+					.IsEnabled(this, &FBakeIndirectWindow::IsBakeEnabled)
+					.OnClicked(this, &FBakeIndirectWindow::OnBakeSelected)
+					[
+						SNew(STextBlock)
+						.Text(NSLOCTEXT("SteamAudio", "BakeSelected", "Bake Selected"))
+						.Font(IDetailLayoutBuilder::GetDetailFont())
+					]
+				]
+			]
+		];
+
+		return SpawnedTab;
+	}
+
+	/**
+	 * Convenience function that computes the cumulative baked data size (in bytes) for a given source ID across the provided probe volumes.
+	 */
+	static int32 ComputeSourceDataSize(const TArray<AActor*>& PhononProbeVolumes, const FName& SourceUID)
+	{
+		int32 SourceDataSize = 0;
+
+		for (auto PhononProbeVolumeActor : PhononProbeVolumes)
+		{
+			auto PhononProbeVolume = Cast<APhononProbeVolume>(PhononProbeVolumeActor);
+			SourceDataSize += PhononProbeVolume->GetDataSizeForSource(SourceUID);
+		}
+
+		return SourceDataSize;
+	}
+
+	/**
+	 * Constructs a table row for the given baked source.
+	 */
+	TSharedRef<ITableRow> FBakeIndirectWindow::OnGenerateBakedSourceRow(TSharedPtr<FBakedSource> Item, const TSharedRef<STableViewBase>& OwnerTable)
+	{
+		auto World = GEditor->GetLevelViewportClients()[0]->GetWorld();
+
+		TArray<AActor*> PhononProbeVolumes;
+		UGameplayStatics::GetAllActorsOfClass(World, APhononProbeVolume::StaticClass(), PhononProbeVolumes);
+
+		Item->DataSize = ComputeSourceDataSize(PhononProbeVolumes, Item->Name);
+
+		return SNew(STableRow<TSharedPtr<FString>>, OwnerTable)
+			.Padding(4)
+			[
+				SNew(SHorizontalBox)
+				+ SHorizontalBox::Slot()
+				[
+					SNew(STextBlock)
+					.Text(FText::FromName(Item->Name))
+					.Font(IDetailLayoutBuilder::GetDetailFont())
+				]
+				+ SHorizontalBox::Slot()
+				[
+					SNew(STextBlock)
+					.Text(FText::AsMemory(Item->DataSize))
+					.Font(IDetailLayoutBuilder::GetDetailFont())
+				]
+			];
+	}
+
+	void FBakeIndirectWindow::OnBakedSourceUpdated(FName UniqueIdentifier)
+	{
+		AsyncTask(ENamedThreads::GameThread, [=]()
+		{
+			RefreshBakedSources();
+			BakedSourcesListView->RequestListRefresh();
+		});
+	}
+
+	/**
+	 * Initiates the bake for selected phonon sources and possibly reverb.
+	 */
+	FReply FBakeIndirectWindow::OnBakeSelected()
+	{
+		// Grab the currently selected items from the baked sources listview
+		TArray<TSharedPtr<FBakedSource>> SelectedSources;
+		BakedSourcesListView->GetSelectedItems(SelectedSources);
+
+		// Collect the phonon source components and determine whether we want to bake reverb
+		TArray<UPhononSourceComponent*> Components;
+		bool ShouldBakeReverb = false;
+		for (const auto& Source : SelectedSources)
+		{
+			if (Source->Name == "__reverb__")
+			{
+				ShouldBakeReverb = true;
+				continue;
+			}
+
+			Components.Add(Source->PhononSourceComponent);
+		}
+
+		// Begin the bake
+		FBakedSourceUpdated Callback;
+		Callback.BindRaw(this, &FBakeIndirectWindow::OnBakedSourceUpdated);
+		Bake(Components, ShouldBakeReverb, Callback);
+
+		return FReply::Handled();
+	}
+
+	/**
+	 * Spawns the window.
+	 */
+	void FBakeIndirectWindow::Invoke()
+	{
+		FGlobalTabmanager::Get()->TryInvokeTab(BakeIndirectTabName);
+	}
+
+	/**
+	 * Populates BakedSources with up-to-date information about all phonon sources and any baked reverb present in the level.
+	 */
+	void FBakeIndirectWindow::RefreshBakedSources()
+	{
+		auto World = GEditor->GetLevelViewportClients()[0]->GetWorld();
+
+		// Clear old data
+		BakedSources.Empty();
+
+		// Grab all probe volumes
+		TArray<AActor*> PhononProbeVolumes;
+		UGameplayStatics::GetAllActorsOfClass(World, APhononProbeVolume::StaticClass(), PhononProbeVolumes);
+
+		// Add baked reverb
+		BakedSources.Add(MakeShareable(new FBakedSource("__reverb__", ComputeSourceDataSize(PhononProbeVolumes, "__reverb__"), nullptr)));
+
+		// Add all baked phonon sources
+		for (TActorIterator<AActor> ActorItr(World); ActorItr; ++ActorItr)
+		{
+			TArray<UActorComponent*> Components;
+			ActorItr->GetComponents(UPhononSourceComponent::StaticClass(), Components);
+			for (auto PhononSourceItr : Components)
+			{
+				if (PhononSourceItr && PhononSourceItr->IsValidLowLevel())
+				{
+					UPhononSourceComponent* PhononSourceComponent = Cast<UPhononSourceComponent>(PhononSourceItr);
+
+					if (PhononSourceComponent && PhononSourceComponent->IsValidLowLevel())
+					{
+						auto SourceDataSize = ComputeSourceDataSize(PhononProbeVolumes, PhononSourceComponent->UniqueIdentifier);
+						BakedSources.Add(MakeShareable(new FBakedSource(PhononSourceComponent->UniqueIdentifier, SourceDataSize, PhononSourceComponent)));
+					}
+				}
+			}
+		}
+	}
+
+	bool FBakeIndirectWindow::IsBakeEnabled() const
+	{
+		return !GIsBaking.Load();
+	}
+
+	//==============================================================================================================================================
+	// FBakedSource
+	//==============================================================================================================================================
+
+	FBakedSource::FBakedSource()
+		: Name("")
+		, DataSize(0)
+		, PhononSourceComponent(nullptr)
+	{}
+
+	FBakedSource::FBakedSource(const FName& Name, const uint32 DataSize, UPhononSourceComponent* PhononSourceComponent)
+		: Name(Name)
+		, DataSize(DataSize)
+		, PhononSourceComponent(PhononSourceComponent)
+	{}
+}
diff --git a/Engine/Plugins/Runtime/Steam/SteamAudio/Source/SteamAudioEditor/Private/BakeIndirectWindow.h b/Engine/Plugins/Runtime/Steam/SteamAudio/Source/SteamAudioEditor/Private/BakeIndirectWindow.h
new file mode 100644
--- /dev/null
+++ b/Engine/Plugins/Runtime/Steam/SteamAudio/Source/SteamAudioEditor/Private/BakeIndirectWindow.h
@@ -0,0 +1,56 @@
+//
+// Copyright (C) Valve Corporation. All rights reserved.
+//
+
+#pragma once
+
+#include "Widgets/Views/SListView.h"
+#include "Templates/SharedPointer.h"
+#include "IndirectBaker.h"
+
+class STableViewBase;
+class UPhononSourceComponent;
+class SDockTab;
+class ITableRow;
+class SButton;
+class FSpawnTabArgs;
+
+namespace SteamAudio
+{
+	/**
+	 * Stores information about a baked source, used for display purposes.
+	 */
+	struct FBakedSource
+	{
+		FName Name;
+		uint32 DataSize;
+		UPhononSourceComponent* PhononSourceComponent;
+
+		FBakedSource();
+		FBakedSource(const FName& Name, const uint32 DataSize, UPhononSourceComponent* PhononSourceComponent);
+	};
+
+	/**
+	 * Provides users with a comprehensive view of all sources that may be baked.
+	 */
+	class FBakeIndirectWindow : public TSharedFromThis<FBakeIndirectWindow>
+	{
+	public:
+		FBakeIndirectWindow();
+		~FBakeIndirectWindow();
+
+		void Invoke();
+		TSharedRef<SDockTab> SpawnTab(const FSpawnTabArgs& TabSpawnArgs);
+		TSharedRef<ITableRow> OnGenerateBakedSourceRow(TSharedPtr<FBakedSource> Item, const TSharedRef<STableViewBase>& OwnerTable);
+		void OnBakedSourceUpdated(FName UniqueIdentifier);
+
+	private:
+		void RefreshBakedSources();
+		FReply OnBakeSelected();
+		bool IsBakeEnabled() const;
+		
+		TSharedPtr<SButton> BakeSelectedButton;
+		TArray<TSharedPtr<FBakedSource>> BakedSources;
+		TSharedPtr<SListView<TSharedPtr<FBakedSource>>> BakedSourcesListView;
+	};
+}
diff --git a/Engine/Plugins/Runtime/Steam/SteamAudio/Source/SteamAudioEditor/Private/IndirectBaker.cpp b/Engine/Plugins/Runtime/Steam/SteamAudio/Source/SteamAudioEditor/Private/IndirectBaker.cpp
new file mode 100644
--- /dev/null
+++ b/Engine/Plugins/Runtime/Steam/SteamAudio/Source/SteamAudioEditor/Private/IndirectBaker.cpp
@@ -0,0 +1,385 @@
+//
+// Copyright (C) Valve Corporation. All rights reserved.
+//
+
+#include "IndirectBaker.h"
+#include "TickableNotification.h"
+#include "PhononSourceComponent.h"
+#include "PhononProbeVolume.h"
+#include "PhononScene.h"
+#include "SteamAudioSettings.h"
+#include "SteamAudioEditorModule.h"
+#include "SteamAudioEnvironment.h"
+
+#include "LevelEditor.h"
+#include "LevelEditorActions.h"
+#include "LevelEditorViewport.h"
+#include "EngineUtils.h"
+#include "Kismet/GameplayStatics.h"
+#include "Async/Async.h"
+#include "Components/AudioComponent.h"
+
+namespace SteamAudio
+{
+	TAtomic<bool> GIsBaking(false);
+
+	static TSharedPtr<FTickableNotification> GBakeTickable = MakeShareable(new FTickableNotification());
+	static int32 GCurrentProbeVolume = 0;
+	static int32 GNumProbeVolumes = 0;
+	static int32 GCurrentBakeTask = 0;
+	static int32 GNumBakeTasks = 0;
+
+	static void BakeProgressCallback(float Progress)
+	{
+		FFormatNamedArguments Arguments;
+		Arguments.Add(TEXT("BakeProgress"), FText::AsPercent(Progress));
+		Arguments.Add(TEXT("CurrentProbeVolume"), FText::AsNumber(GCurrentProbeVolume));
+		Arguments.Add(TEXT("NumProbeVolumes"), FText::AsNumber(GNumProbeVolumes));
+		Arguments.Add(TEXT("NumBakeTasks"), FText::AsNumber(GNumBakeTasks));
+		Arguments.Add(TEXT("CurrentBakeTask"), FText::AsNumber(GCurrentBakeTask));
+		GBakeTickable->SetDisplayText(FText::Format(NSLOCTEXT("SteamAudio", "BakeProgressFmt", "Baking {CurrentBakeTask}/{NumBakeTasks} sources \n {CurrentProbeVolume}/{NumProbeVolumes} probe volumes ({BakeProgress} complete)"), Arguments));
+	}
+
+	static void CancelBake()
+	{
+		iplCancelBake();
+		GIsBaking = false;
+	}
+
+	/**
+	 * Bakes propagation for all sources in PhononSourceComponents. Bakes reverb if BakeReverb is set. Performs baking across all probe volumes
+	 * in the scene. Runs baking in an async task so that UI remains responsive.
+	 */
+	void Bake(const TArray<UPhononSourceComponent*> PhononSourceComponents, const bool BakeReverb, FBakedSourceUpdated BakedSourceUpdated)
+	{
+		GIsBaking = true;
+
+		GBakeTickable->SetDisplayText(NSLOCTEXT("SteamAudio", "Baking", "Baking..."));
+		GBakeTickable->CreateNotificationWithCancel(FSimpleDelegate::CreateStatic(CancelBake));
+
+		auto World = GEditor->GetLevelViewportClients()[0]->GetWorld();
+		check(World);
+
+		GNumBakeTasks = BakeReverb ? PhononSourceComponents.Num() + 1 : PhononSourceComponents.Num();
+		GCurrentBakeTask = 1;
+
+		// Get all probe volumes (cannot do this in the async task - not on game thread)
+		TArray<AActor*> PhononProbeVolumes;
+		UGameplayStatics::GetAllActorsOfClass(World, APhononProbeVolume::StaticClass(), PhononProbeVolumes);
+
+		Async(EAsyncExecution::Thread, [=]()
+		{
+			// Ensure we have at least one probe
+			bool AtLeastOneProbe = false;
+
+			for (auto PhononProbeVolumeActor : PhononProbeVolumes)
+			{
+				auto PhononProbeVolume = Cast<APhononProbeVolume>(PhononProbeVolumeActor);
+				if (PhononProbeVolume->NumProbes > 0)
+				{
+					AtLeastOneProbe = true;
+					break;
+				}
+			}
+
+			if (!AtLeastOneProbe)
+			{
+				UE_LOG(LogSteamAudioEditor, Error, TEXT("Ensure at least one Phonon Probe Volume with probes exists."));
+				GBakeTickable->SetDisplayText(NSLOCTEXT("SteamAudio", "BakeFailed_NoProbes", "Bake failed. Create at least one Phonon Probe Volume that has probes."));
+				GBakeTickable->DestroyNotification(SNotificationItem::CS_Fail);
+				GIsBaking = false;
+				return;
+			}
+
+			IPLBakingSettings BakingSettings;
+			BakingSettings.bakeParametric = IPL_FALSE;
+			BakingSettings.bakeConvolution = IPL_TRUE;
+
+			IPLhandle ComputeDevice = nullptr;
+
+			IPLComputeDeviceFilter DeviceFilter;
+			DeviceFilter.fractionCUsForIRUpdate = GetDefault<USteamAudioSettings>()->GetFractionComputeUnitsForIRUpdate();
+			DeviceFilter.maxCUsToReserve = GetDefault<USteamAudioSettings>()->MaxComputeUnits;
+			DeviceFilter.type = IPL_COMPUTEDEVICE_GPU;
+
+			IPLSimulationSettings SimulationSettings = GetDefault<USteamAudioSettings>()->GetBakedSimulationSettings();
+
+			// If we are using RadeonRays, attempt to create a compute device.
+			if (GetDefault<USteamAudioSettings>()->RayTracer == EIplRayTracerType::RADEONRAYS)
+			{
+				UE_LOG(LogSteamAudioEditor, Log, TEXT("Using Radeon Rays - creating GPU compute device..."));
+
+				IPLerror Error = iplCreateComputeDevice(GlobalContext, DeviceFilter, &ComputeDevice);
+
+				// If we failed to create a compute device, fall back to Phonon scene.
+				if (Error != IPL_STATUS_SUCCESS)
+				{
+					UE_LOG(LogSteamAudioEditor, Warning, TEXT("Failed to create GPU compute device, falling back to Phonon."));
+
+					SimulationSettings.sceneType = IPL_SCENETYPE_PHONON;
+					SimulationSettings.bakingBatchSize = 1;
+				}
+			}
+
+			IPLhandle PhononScene = nullptr;
+			IPLhandle PhononEnvironment = nullptr;
+			FPhononSceneInfo PhononSceneInfo;
+
+			GBakeTickable->SetDisplayText(NSLOCTEXT("SteamAudio", "LoadingScene", "Loading scene..."));
+
+			// Load the scene
+			if (!LoadSceneFromDisk(World, ComputeDevice, SimulationSettings, &PhononScene, PhononSceneInfo, nullptr))
+			{
+				// If we can't find the scene, then presumably they haven't generated probes either, so just exit
+				UE_LOG(LogSteamAudioEditor, Error, TEXT("Unable to create Phonon environment: .phononscene not found. Be sure to export the scene."));
+
+				if (ComputeDevice)
+				{
+					iplDestroyComputeDevice(&ComputeDevice);
+				}
+
+				GBakeTickable->SetDisplayText(NSLOCTEXT("SteamAudio", "BakeFailed_NoScene", "Bake failed. Export scene first."));
+				GBakeTickable->DestroyNotification(SNotificationItem::CS_Fail);
+				GIsBaking = false;
+				return;
+			}
+
+			iplCreateEnvironment(SteamAudio::GlobalContext, ComputeDevice, SimulationSettings, PhononScene, nullptr, &PhononEnvironment);
+
+			if (BakeReverb)
+			{
+				GBakeTickable->SetDisplayText(NSLOCTEXT("SteamAudio", "Baking", "Baking..."));
+				GNumProbeVolumes = PhononProbeVolumes.Num();
+				GCurrentProbeVolume = 1;
+
+				for (auto PhononProbeVolumeActor : PhononProbeVolumes)
+				{
+					auto PhononProbeVolume = Cast<APhononProbeVolume>(PhononProbeVolumeActor);
+
+					IPLhandle ProbeBox = nullptr;
+					PhononProbeVolume->LoadProbeBoxFromDisk(&ProbeBox);
+
+					IPLBakedDataIdentifier ReverbIdentifier;
+					ReverbIdentifier.identifier = 0;
+					ReverbIdentifier.type = IPL_BAKEDDATATYPE_REVERB;
+
+					iplDeleteBakedDataByIdentifier(ProbeBox, ReverbIdentifier);
+					iplBakeReverb(PhononEnvironment, ProbeBox, BakingSettings, BakeProgressCallback);
+
+					if (!GIsBaking.Load())
+					{
+						iplDestroyProbeBox(&ProbeBox);
+						break;
+					}
+
+					FBakedDataInfo BakedDataInfo;
+					BakedDataInfo.Name = "__reverb__";
+					BakedDataInfo.Size = iplGetBakedDataSizeByIdentifier(ProbeBox, ReverbIdentifier);
+
+					auto ExistingInfo = PhononProbeVolume->BakedDataInfo.FindByPredicate([=](const FBakedDataInfo& InfoItem)
+					{
+						return InfoItem.Name == BakedDataInfo.Name;
+					});
+
+					if (ExistingInfo)
+					{
+						ExistingInfo->Size = BakedDataInfo.Size;
+					}
+					else
+					{
+						PhononProbeVolume->BakedDataInfo.Add(BakedDataInfo);
+						PhononProbeVolume->BakedDataInfo.Sort();
+					}
+
+					PhononProbeVolume->UpdateProbeData(ProbeBox);
+					iplDestroyProbeBox(&ProbeBox);
+					++GCurrentProbeVolume;
+				}
+
+				if (!GIsBaking.Load())
+				{
+					iplDestroyEnvironment(&PhononEnvironment);
+					iplDestroyScene(&PhononScene);
+
+					if (ComputeDevice)
+					{
+						iplDestroyComputeDevice(&ComputeDevice);
+					}
+
+					GBakeTickable->SetDisplayText(NSLOCTEXT("SteamAudio", "BakeCancelled", "Bake cancelled."));
+					GBakeTickable->DestroyNotification(SNotificationItem::CS_Fail);
+					return;
+				}
+
+				BakedSourceUpdated.ExecuteIfBound("__reverb__");
+
+				++GCurrentBakeTask;
+			}
+
+			// IN PROGRESS
+			FIdentifierMap BakedIdentifierMap;
+			LoadBakedIdentifierMapFromDisk(World, BakedIdentifierMap);
+
+			for (UPhononSourceComponent* PhononSourceComponent : PhononSourceComponents)
+			{
+				// Check if phonon source is valid
+				if (!PhononSourceComponent || !PhononSourceComponent->IsValidLowLevelFast())
+				{
+					UE_LOG(LogSteamAudioEditor, Warning, TEXT("Phonon Source Component is invalid. It will be skipped."));
+					continue;
+				}
+				else if (!PhononSourceComponent->GetOwner() && !PhononSourceComponent->GetOwner()->IsValidLowLevelFast())
+				{
+					UE_LOG(LogSteamAudioEditor, Warning, TEXT("Phonon Source Component \"%s\" is not attached to a valid actor. It will be skipped."), *(PhononSourceComponent->UniqueIdentifier.ToString()));
+					continue;
+				}
+				else
+				{
+					UE_LOG(LogSteamAudioEditor, Log, TEXT("Phonon Source Component \"%s\" was used to trigger this bake task."), *(PhononSourceComponent->UniqueIdentifier.ToString()));
+				}
+
+				// Go through all audio components that the owner of this phonon source actor has
+				TArray<FString> BakeAudioIdentifiers;
+				TArray<UActorComponent*> Components;
+				PhononSourceComponent->GetOwner()->GetComponents(UAudioComponent::StaticClass(), Components);
+
+				for (auto ActorComp : Components)
+				{
+					if (ActorComp && ActorComp->IsValidLowLevel())
+					{
+						UAudioComponent* AudioComponent = Cast<UAudioComponent>(ActorComp);
+
+						// Set the User ID on the audio component
+						if (AudioComponent && AudioComponent->IsValidLowLevelFast())
+						{
+							// Apply the phonon source id to this audio component
+							AudioComponent->AudioComponentUserID = PhononSourceComponent->UniqueIdentifier;
+							//UE_LOG(LogTemp, Warning, TEXT("Audio Component [%s] UserID is now set to [%s]"), *AudioComponent->GetFullName(), *AudioComponent->AudioComponentUserID.ToString());
+
+							// Check if there's a phonon source child for this audio component
+							for (auto ChildComp : AudioComponent->GetAttachChildren())
+							{
+								auto* ChildPhononComponent = Cast<UPhononSourceComponent>(ChildComp);
+
+								// If there is one, then apply it's child phonon source id, overriding and phonon source id already applied to it
+								if (ChildPhononComponent && ChildPhononComponent->IsValidLowLevelFast())
+								{
+									AudioComponent->AudioComponentUserID = ChildPhononComponent->UniqueIdentifier;
+									//UE_LOG(LogTemp, Warning, TEXT("Audio Component [%s] UserID is now overriden to [%s]"), *AudioComponent->GetFullName(), *AudioComponent->AudioComponentUserID.ToString());
+									break;
+								}
+							}
+
+							// Apply id to bake map if not present
+							FString SourceString = AudioComponent->GetAudioComponentUserID().ToString().ToLower();
+							if (!BakedIdentifierMap.ContainsKey(SourceString))
+							{
+								BakedIdentifierMap.Add(SourceString);
+							}
+
+							BakeAudioIdentifiers.AddUnique(SourceString);
+						}
+					}
+				}
+
+				// Bake audio
+				if (BakeAudioIdentifiers.Num() > 0)
+				{
+					for (auto BakeAudioIdentifier : BakeAudioIdentifiers)
+					{
+						IPLBakedDataIdentifier SourceIdentifier;
+						SourceIdentifier.type = IPL_BAKEDDATATYPE_STATICSOURCE;
+						SourceIdentifier.identifier = BakedIdentifierMap.Get(BakeAudioIdentifier);
+
+						FString BakingText = FString::Printf(TEXT("Baking Audio %s"), *BakeAudioIdentifier);
+						GBakeTickable->SetDisplayText(FText::AsCultureInvariant(BakingText));
+						GNumProbeVolumes = PhononProbeVolumes.Num();
+						GCurrentProbeVolume = 1;
+
+						UE_LOG(LogTemp, Log, TEXT("Baking phonon source component %s"), *BakeAudioIdentifier);
+
+						for (auto PhononProbeVolumeActor : PhononProbeVolumes)
+						{
+							auto PhononProbeVolume = Cast<APhononProbeVolume>(PhononProbeVolumeActor);
+
+							IPLhandle ProbeBox = nullptr;
+							PhononProbeVolume->LoadProbeBoxFromDisk(&ProbeBox);
+
+							IPLSphere SourceInfluence;
+							SourceInfluence.radius = PhononSourceComponent->BakingRadius * SteamAudio::SCALEFACTOR;
+							SourceInfluence.center = SteamAudio::UnrealToPhononIPLVector3(PhononSourceComponent->GetComponentLocation());
+
+							iplDeleteBakedDataByIdentifier(ProbeBox, SourceIdentifier);
+							iplBakePropagation(PhononEnvironment, ProbeBox, SourceInfluence, SourceIdentifier, BakingSettings, BakeProgressCallback);
+
+							if (!GIsBaking.Load())
+							{
+								iplDestroyProbeBox(&ProbeBox);
+								break;
+							}
+
+							FBakedDataInfo BakedDataInfo;
+							BakedDataInfo.Name = PhononSourceComponent->UniqueIdentifier;
+							BakedDataInfo.Size = iplGetBakedDataSizeByIdentifier(ProbeBox, SourceIdentifier);
+
+							auto ExistingInfo = PhononProbeVolume->BakedDataInfo.FindByPredicate([=](const FBakedDataInfo& InfoItem)
+							{
+								return InfoItem.Name == BakedDataInfo.Name;
+							});
+
+							if (ExistingInfo)
+							{
+								ExistingInfo->Size = BakedDataInfo.Size;
+							}
+							else
+							{
+								PhononProbeVolume->BakedDataInfo.Add(BakedDataInfo);
+								PhononProbeVolume->BakedDataInfo.Sort();
+							}
+
+							PhononProbeVolume->UpdateProbeData(ProbeBox);
+							iplDestroyProbeBox(&ProbeBox);
+							++GCurrentProbeVolume;
+						}
+
+						const bool bIsBaking = GIsBaking.Load();
+						if (!bIsBaking)
+						{
+							break;
+						}
+
+						BakedSourceUpdated.ExecuteIfBound(PhononSourceComponent->UniqueIdentifier);
+
+						++GCurrentBakeTask;
+					}
+				}
+				else
+				{
+					UE_LOG(LogSteamAudioEditor, Warning, TEXT("Actor containing the Phonon source \"%s\" has no Audio Component. It will be skipped."), *(PhononSourceComponent->UniqueIdentifier.ToString()));
+				}
+			}
+
+			SaveBakedIdentifierMapToDisk(World, BakedIdentifierMap);
+
+			iplDestroyEnvironment(&PhononEnvironment);
+			iplDestroyScene(&PhononScene);
+
+			if (ComputeDevice)
+			{
+				iplDestroyComputeDevice(&ComputeDevice);
+			}
+
+			if (!GIsBaking.Load())
+			{
+				GBakeTickable->SetDisplayText(NSLOCTEXT("SteamAudio", "BakeCancelled", "Bake cancelled."));
+				GBakeTickable->DestroyNotification(SNotificationItem::CS_Fail);
+			}
+			else
+			{
+				GBakeTickable->SetDisplayText(NSLOCTEXT("SteamAudio", "BakePropagationComplete", "Bake propagation complete."));
+				GBakeTickable->DestroyNotification();
+				GIsBaking = false;
+			}
+		});
+	}
+}
diff --git a/Engine/Plugins/Runtime/Steam/SteamAudio/Source/SteamAudioEditor/Private/IndirectBaker.h b/Engine/Plugins/Runtime/Steam/SteamAudio/Source/SteamAudioEditor/Private/IndirectBaker.h
new file mode 100644
--- /dev/null
+++ b/Engine/Plugins/Runtime/Steam/SteamAudio/Source/SteamAudioEditor/Private/IndirectBaker.h
@@ -0,0 +1,20 @@
+//
+// Copyright (C) Valve Corporation. All rights reserved.
+//
+
+#pragma once
+
+#include "Delegates/Delegate.h"
+#include "CoreMinimal.h"
+
+class UPhononSourceComponent;
+
+namespace SteamAudio
+{
+	// True if a baking process is currently running.
+	extern TAtomic<bool> GIsBaking;
+
+	DECLARE_DELEGATE_OneParam(FBakedSourceUpdated, FName);
+	
+	void Bake(const TArray<UPhononSourceComponent*> PhononSourceComponents, const bool BakeReverb, SteamAudio::FBakedSourceUpdated BakedSourceUpdated);
+}
diff --git a/Engine/Plugins/Runtime/Steam/SteamAudio/Source/SteamAudioEditor/Private/PhononOcclusionSettingsFactory.cpp b/Engine/Plugins/Runtime/Steam/SteamAudio/Source/SteamAudioEditor/Private/PhononOcclusionSettingsFactory.cpp
new file mode 100644
--- /dev/null
+++ b/Engine/Plugins/Runtime/Steam/SteamAudio/Source/SteamAudioEditor/Private/PhononOcclusionSettingsFactory.cpp
@@ -0,0 +1,64 @@
+//
+// Copyright (C) Valve Corporation. All rights reserved.
+//
+
+#include "PhononOcclusionSettingsFactory.h"
+
+#include "AudioAnalytics.h"
+#include "PhononOcclusionSourceSettings.h"
+#include "UObject/ObjectMacros.h"
+#include "UObject/Object.h"
+
+namespace SteamAudio
+{
+	FText FAssetTypeActions_PhononOcclusionSettings::GetName() const
+	{
+		return NSLOCTEXT("SteamAudio", "AssetTypeActions_PhononOcclusionSettings", "Phonon Source Occlusion Settings"); 
+	}
+
+	FColor FAssetTypeActions_PhononOcclusionSettings::GetTypeColor() const
+	{
+		return FColor(245, 195, 101);
+	}
+
+	UClass* FAssetTypeActions_PhononOcclusionSettings::GetSupportedClass() const
+	{
+		return UPhononOcclusionSourceSettings::StaticClass();
+	}
+
+	uint32 FAssetTypeActions_PhononOcclusionSettings::GetCategories()
+	{
+		return EAssetTypeCategories::Sounds;
+	}
+
+	const TArray<FText>& FAssetTypeActions_PhononOcclusionSettings::GetSubMenus() const
+	{
+		static const TArray<FText> PhononSubMenus
+		{
+			NSLOCTEXT("SteamAudio", "AssetPhononSubMenu", "Phonon")
+		};
+		return PhononSubMenus;
+	}
+}
+
+UPhononOcclusionSettingsFactory::UPhononOcclusionSettingsFactory(const FObjectInitializer& ObjectInitializer)
+	: Super(ObjectInitializer)
+{
+	SupportedClass = UPhononOcclusionSourceSettings::StaticClass();
+
+	bCreateNew = true;
+	bEditorImport = false;
+	bEditAfterNew = true;
+}
+
+UObject* UPhononOcclusionSettingsFactory::FactoryCreateNew(UClass* Class, UObject* InParent, FName InName, EObjectFlags Flags, UObject* Context,
+	FFeedbackContext* Warn)
+{
+	Audio::Analytics::RecordEvent_Usage(TEXT("SteamAudio.PhononOcclusionSettingsCreated"));
+	return NewObject<UPhononOcclusionSourceSettings>(InParent, InName, Flags);
+}
+
+uint32 UPhononOcclusionSettingsFactory::GetMenuCategories() const
+{
+	return EAssetTypeCategories::Sounds;
+}
\ No newline at end of file
diff --git a/Engine/Plugins/Runtime/Steam/SteamAudio/Source/SteamAudioEditor/Private/PhononOcclusionSettingsFactory.h b/Engine/Plugins/Runtime/Steam/SteamAudio/Source/SteamAudioEditor/Private/PhononOcclusionSettingsFactory.h
new file mode 100644
--- /dev/null
+++ b/Engine/Plugins/Runtime/Steam/SteamAudio/Source/SteamAudioEditor/Private/PhononOcclusionSettingsFactory.h
@@ -0,0 +1,36 @@
+//
+// Copyright (C) Valve Corporation. All rights reserved.
+//
+
+#pragma once
+
+#include "IAudioExtensionPlugin.h"
+#include "AssetToolsModule.h"
+#include "Factories/Factory.h"
+#include "AssetTypeActions_Base.h"
+#include "PhononOcclusionSettingsFactory.generated.h"
+
+namespace SteamAudio
+{
+	class FAssetTypeActions_PhononOcclusionSettings : public FAssetTypeActions_Base
+	{
+	public:
+		virtual FText GetName() const override;
+		virtual FColor GetTypeColor() const override;
+		virtual UClass* GetSupportedClass() const override;
+		virtual uint32 GetCategories() override;
+
+		virtual const TArray<FText>& GetSubMenus() const override;
+	};
+}
+
+UCLASS(MinimalAPI, hidecategories = Object)
+class UPhononOcclusionSettingsFactory : public UFactory
+{
+	GENERATED_UCLASS_BODY()
+
+	virtual UObject* FactoryCreateNew(UClass* Class, UObject* InParent, FName Name, EObjectFlags Flags, UObject* Context,
+		FFeedbackContext* Warn) override;
+
+	virtual uint32 GetMenuCategories() const override;
+};
diff --git a/Engine/Plugins/Runtime/Steam/SteamAudio/Source/SteamAudioEditor/Private/PhononProbeComponentVisualizer.cpp b/Engine/Plugins/Runtime/Steam/SteamAudio/Source/SteamAudioEditor/Private/PhononProbeComponentVisualizer.cpp
new file mode 100644
--- /dev/null
+++ b/Engine/Plugins/Runtime/Steam/SteamAudio/Source/SteamAudioEditor/Private/PhononProbeComponentVisualizer.cpp
@@ -0,0 +1,25 @@
+//
+// Copyright (C) Valve Corporation. All rights reserved.
+//
+
+#include "PhononProbeComponentVisualizer.h"
+#include "PhononProbeComponent.h"
+#include "SceneView.h"
+#include "SceneManagement.h"
+
+namespace SteamAudio
+{
+	void FPhononProbeComponentVisualizer::DrawVisualization(const UActorComponent* Component, const FSceneView* View, FPrimitiveDrawInterface* PDI)
+	{
+		UPhononProbeComponent* ProbeComponent = const_cast<UPhononProbeComponent*>(Cast<UPhononProbeComponent>(Component));
+		if (ProbeComponent)
+		{
+			FScopeLock ProbeLocationsLock(&ProbeComponent->ProbeLocationsCriticalSection);
+			const FColor Color(0, 153, 255);
+			for (const auto& ProbeLocation : ProbeComponent->ProbeLocations)
+			{
+				PDI->DrawPoint(ProbeLocation, Color, 5, SDPG_World);
+			}
+		}
+	}
+}
\ No newline at end of file
diff --git a/Engine/Plugins/Runtime/Steam/SteamAudio/Source/SteamAudioEditor/Private/PhononProbeComponentVisualizer.h b/Engine/Plugins/Runtime/Steam/SteamAudio/Source/SteamAudioEditor/Private/PhononProbeComponentVisualizer.h
new file mode 100644
--- /dev/null
+++ b/Engine/Plugins/Runtime/Steam/SteamAudio/Source/SteamAudioEditor/Private/PhononProbeComponentVisualizer.h
@@ -0,0 +1,19 @@
+//
+// Copyright (C) Valve Corporation. All rights reserved.
+//
+
+#pragma once
+
+#include "ComponentVisualizer.h"
+
+class FPrimitiveDrawInterface;
+class FSceneView;
+
+namespace SteamAudio
+{
+	class FPhononProbeComponentVisualizer : public FComponentVisualizer
+	{
+	public:
+		virtual void DrawVisualization(const UActorComponent* Component, const FSceneView* View, FPrimitiveDrawInterface* PDI) override;
+	};
+}
diff --git a/Engine/Plugins/Runtime/Steam/SteamAudio/Source/SteamAudioEditor/Private/PhononProbeVolumeDetails.cpp b/Engine/Plugins/Runtime/Steam/SteamAudio/Source/SteamAudioEditor/Private/PhononProbeVolumeDetails.cpp
new file mode 100644
--- /dev/null
+++ b/Engine/Plugins/Runtime/Steam/SteamAudio/Source/SteamAudioEditor/Private/PhononProbeVolumeDetails.cpp
@@ -0,0 +1,291 @@
+//
+// Copyright (C) Valve Corporation. All rights reserved.
+//
+
+#include "PhononProbeVolumeDetails.h"
+
+#include "PhononCommon.h"
+#include "SteamAudioEditorModule.h"
+#include "SteamAudioSettings.h"
+#include "SteamAudioEnvironment.h"
+
+#include "DetailLayoutBuilder.h"
+#include "DetailCategoryBuilder.h"
+#include "DetailWidgetRow.h"
+#include "IDetailsView.h"
+#include "IDetailCustomization.h"
+#include "PhononProbeVolume.h"
+#include "PhononProbeComponent.h"
+#include "TickableNotification.h"
+#include "Async/Async.h"
+#include "PhononScene.h"
+#include "LevelEditorViewport.h"
+#include "Settings/LevelEditorPlaySettings.h"
+#include "Settings/LevelEditorViewportSettings.h"
+#include "PropertyCustomizationHelpers.h"
+#include "IDetailChildrenBuilder.h"
+#include "IDetailPropertyRow.h"
+#include "Widgets/Input/SButton.h"
+#include "Widgets/Text/STextBlock.h"
+#include "Widgets/Layout/SBox.h"
+#include "Editor.h"
+
+namespace SteamAudio
+{
+	static TSharedPtr<FTickableNotification> GGenerateProbesTickable = MakeShareable(new FTickableNotification());
+
+	static void GenerateProbesProgressCallback(float Progress)
+	{
+		FFormatNamedArguments Arguments;
+		Arguments.Add(TEXT("GenerateProbesProgress"), FText::AsPercent(Progress));
+		GGenerateProbesTickable->SetDisplayText(FText::Format(NSLOCTEXT("SteamAudio", "ComputingProbeLocationsText", "Computing probe locations ({GenerateProbesProgress} complete)"), Arguments));
+	}
+
+	TSharedRef<IDetailCustomization> FPhononProbeVolumeDetails::MakeInstance()
+	{
+		return MakeShareable(new FPhononProbeVolumeDetails);
+	}
+
+	FText FPhononProbeVolumeDetails::GetTotalDataSize()
+	{
+		return FText::AsMemory(PhononProbeVolume->ProbeDataSize);
+	}
+
+	void FPhononProbeVolumeDetails::CustomizeDetails(IDetailLayoutBuilder& DetailLayout)
+	{
+		const TArray<TWeakObjectPtr<UObject>>& SelectedObjects = DetailLayout.GetSelectedObjects();
+
+		for (int32 ObjectIndex = 0; ObjectIndex < SelectedObjects.Num(); ++ObjectIndex)
+		{
+			const TWeakObjectPtr<UObject>& CurrentObject = SelectedObjects[ObjectIndex];
+			if (CurrentObject.IsValid())
+			{
+				APhononProbeVolume* CurrentCaptureActor = Cast<APhononProbeVolume>(CurrentObject.Get());
+				if (CurrentCaptureActor)
+				{
+					PhononProbeVolume = CurrentCaptureActor;
+					break;
+				}
+			}
+		}
+
+		DetailLayout.HideCategory("BrushSettings");
+
+		DetailLayout.EditCategory("ProbeGeneration").AddProperty(GET_MEMBER_NAME_CHECKED(APhononProbeVolume, PlacementStrategy));
+		DetailLayout.EditCategory("ProbeGeneration").AddProperty(GET_MEMBER_NAME_CHECKED(APhononProbeVolume, HorizontalSpacing));
+		DetailLayout.EditCategory("ProbeGeneration").AddProperty(GET_MEMBER_NAME_CHECKED(APhononProbeVolume, HeightAboveFloor));
+		DetailLayout.EditCategory("ProbeGeneration").AddCustomRow(NSLOCTEXT("SteamAudio", "GenerateProbes", "Generate Probes"))
+			.NameContent()
+			[
+				SNullWidget::NullWidget
+			]
+			.ValueContent()
+			[
+				SNew(SHorizontalBox)
+				+ SHorizontalBox::Slot()
+				.AutoWidth()
+				[
+					SNew(SButton)
+					.ContentPadding(2)
+					.VAlign(VAlign_Center)
+					.HAlign(HAlign_Center)
+					.OnClicked(this, &FPhononProbeVolumeDetails::OnGenerateProbes)
+					[
+						SNew(STextBlock)
+						.Text(NSLOCTEXT("SteamAudio", "GenerateProbes", "Generate Probes"))
+						.Font(IDetailLayoutBuilder::GetDetailFont())
+					]
+				]
+			];
+
+			TSharedPtr<IPropertyHandle> BakedDataProperty = DetailLayout.GetProperty(GET_MEMBER_NAME_CHECKED(APhononProbeVolume, BakedDataInfo));
+			TSharedRef<FDetailArrayBuilder> BakedDataBuilder = MakeShareable(new FDetailArrayBuilder(BakedDataProperty.ToSharedRef()));
+			BakedDataBuilder->OnGenerateArrayElementWidget(FOnGenerateArrayElementWidget::CreateSP(this, &FPhononProbeVolumeDetails::OnGenerateBakedDataInfo));
+			DetailLayout.EditCategory("ProbeVolumeStatistics").AddProperty(GET_MEMBER_NAME_CHECKED(APhononProbeVolume, NumProbes));
+
+			auto ProbeDataSize = DetailLayout.GetProperty(GET_MEMBER_NAME_CHECKED(APhononProbeVolume, ProbeDataSize));
+			TAttribute<FText> TotalDataSize = TAttribute<FText>::Create(TAttribute<FText>::FGetter::CreateSP(this, &FPhononProbeVolumeDetails::GetTotalDataSize));
+		
+			DetailLayout.EditCategory("ProbeVolumeStatistics").AddProperty(ProbeDataSize).CustomWidget()
+				.NameContent()
+				[
+					SNew(STextBlock)
+					.Text(NSLOCTEXT("SteamAudio", "ProbeDataSize", "Probe Data Size"))
+					.Font(IDetailLayoutBuilder::GetDetailFont())
+				]
+				.ValueContent()
+				[
+					SNew(STextBlock)
+					.Text(TotalDataSize)
+					.Font(IDetailLayoutBuilder::GetDetailFont())
+				];
+			DetailLayout.EditCategory("ProbeVolumeStatistics").AddCustomBuilder(BakedDataBuilder);
+	}
+
+	void FPhononProbeVolumeDetails::OnGenerateBakedDataInfo(TSharedRef<IPropertyHandle> PropertyHandle, int32 ArrayIndex, IDetailChildrenBuilder& ChildrenBuilder)
+	{
+		auto& BakedDataRow = ChildrenBuilder.AddProperty(PropertyHandle);
+		auto& BakedDataInfo = PhononProbeVolume->BakedDataInfo[ArrayIndex];
+
+		BakedDataRow.ShowPropertyButtons(false);
+		BakedDataRow.CustomWidget(false)
+			.NameContent()
+			[
+				SNew(STextBlock)
+				.Text(FText::FromName(BakedDataInfo.Name))
+				.Font(IDetailLayoutBuilder::GetDetailFont())
+			]
+			.ValueContent()
+			[
+				SNew(SHorizontalBox)
+				+ SHorizontalBox::Slot()
+				[
+					SNew(SBox)
+					.MinDesiredWidth(200)
+					.VAlign(VAlign_Center)
+					[
+						SNew(STextBlock)
+						.Text(FText::AsMemory(BakedDataInfo.Size))
+						.Font(IDetailLayoutBuilder::GetDetailFont())
+					]
+				]
+				+ SHorizontalBox::Slot()
+					.AutoWidth()
+				[
+					PropertyCustomizationHelpers::MakeDeleteButton(FSimpleDelegate::CreateSP(this, &FPhononProbeVolumeDetails::OnClearBakedDataClicked, ArrayIndex))
+				]
+			];
+	}
+
+	void FPhononProbeVolumeDetails::OnClearBakedDataClicked(const int32 ArrayIndex)
+	{
+		IPLhandle ProbeBox = nullptr;
+		PhononProbeVolume->LoadProbeBoxFromDisk(&ProbeBox);
+		
+		FIdentifierMap IdentifierMap;
+		LoadBakedIdentifierMapFromDisk(PhononProbeVolume->GetWorld(), IdentifierMap);
+		
+		FString BakedDataString = PhononProbeVolume->BakedDataInfo[ArrayIndex].Name.ToString().ToLower();
+
+		IPLBakedDataIdentifier BakedDataIdentifier;
+		BakedDataIdentifier.type = BakedDataString.Equals("__reverb__") ? IPL_BAKEDDATATYPE_REVERB : IPL_BAKEDDATATYPE_STATICSOURCE;
+		BakedDataIdentifier.identifier = IdentifierMap.Get(BakedDataString);
+		
+		iplDeleteBakedDataByIdentifier(ProbeBox, BakedDataIdentifier);
+		PhononProbeVolume->BakedDataInfo.RemoveAt(ArrayIndex);
+		PhononProbeVolume->UpdateProbeData(ProbeBox);
+		iplDestroyProbeBox(&ProbeBox);
+	}
+
+	FReply FPhononProbeVolumeDetails::OnGenerateProbes()
+	{
+		GGenerateProbesTickable->SetDisplayText(NSLOCTEXT("SteamAudio", "GeneratingProbes", "Generating probes..."));
+		GGenerateProbesTickable->CreateNotification();
+
+		// Grab a copy of the volume ptr as it will be destroyed if user clicks off of volume in GUI
+		auto PhononProbeVolumeHandle = PhononProbeVolume.Get();
+
+		Async(EAsyncExecution::Thread, [PhononProbeVolumeHandle]()
+		{
+			// Load the scene
+			UWorld* World = GEditor->GetLevelViewportClients()[0]->GetWorld();
+			IPLhandle PhononScene = nullptr;
+			FPhononSceneInfo PhononSceneInfo;
+
+			GGenerateProbesTickable->SetDisplayText(NSLOCTEXT("SteamAudio", "LoadingScene", "Loading scene..."));
+
+			IPLhandle ComputeDevice = nullptr;
+
+			IPLComputeDeviceFilter DeviceFilter;
+			DeviceFilter.fractionCUsForIRUpdate = GetDefault<USteamAudioSettings>()->GetFractionComputeUnitsForIRUpdate();
+			DeviceFilter.maxCUsToReserve = GetDefault<USteamAudioSettings>()->MaxComputeUnits;
+			DeviceFilter.type = IPL_COMPUTEDEVICE_GPU;
+
+			IPLSimulationSettings SimulationSettings = GetDefault<USteamAudioSettings>()->GetBakedSimulationSettings();
+
+			// If we are using RadeonRays, attempt to create a compute device.
+			if (GetDefault<USteamAudioSettings>()->RayTracer == EIplRayTracerType::RADEONRAYS)
+			{
+				UE_LOG(LogSteamAudioEditor, Log, TEXT("Using Radeon Rays - creating GPU compute device..."));
+
+				IPLerror Error = iplCreateComputeDevice(GlobalContext, DeviceFilter, &ComputeDevice);
+
+				// If we failed to create a compute device, fall back to Phonon scene.
+				if (Error != IPL_STATUS_SUCCESS)
+				{
+					UE_LOG(LogSteamAudioEditor, Warning, TEXT("Failed to create GPU compute device, falling back to Phonon."));
+
+					SimulationSettings.sceneType = IPL_SCENETYPE_PHONON;
+					SimulationSettings.bakingBatchSize = 1;
+				}
+			}
+
+			// Attempt to load from disk, otherwise export
+			if (!LoadSceneFromDisk(World, ComputeDevice, SimulationSettings, &PhononScene, PhononSceneInfo, nullptr))
+			{
+				IPLhandle PhononStaticMesh = nullptr;
+
+				if (!CreateScene(World, &PhononScene, &PhononStaticMesh, PhononSceneInfo.NumTriangles, PhononSceneInfo.NumDynTriangles, PhononSceneInfo.DynDataSize))
+				{
+					GGenerateProbesTickable->QueueWorkItem(FWorkItem([](FText& DisplayText) {
+						DisplayText = NSLOCTEXT("SteamAudio", "UnableToCreateScene", "Unable to create scene.");
+					}, SNotificationItem::CS_Fail, true));
+
+					if (ComputeDevice)
+					{
+						iplDestroyComputeDevice(&ComputeDevice);
+					}
+
+					return;
+				}
+
+				PhononSceneInfo.DataSize = iplSaveScene(PhononScene, nullptr);
+				bool SaveSceneSuccessful = SaveFinalizedSceneToDisk(World, PhononScene, PhononSceneInfo);
+
+				if (SaveSceneSuccessful)
+				{
+					FSteamAudioEditorModule* Module = &FModuleManager::GetModuleChecked<FSteamAudioEditorModule>("SteamAudioEditor");
+					if (Module != nullptr)
+					{
+						Module->SetCurrentPhononSceneInfo(PhononSceneInfo);
+					}
+				}
+
+				iplDestroyStaticMesh(&PhononStaticMesh);
+			}
+
+			// Place probes
+			IPLhandle SceneCopy = PhononScene;
+			TArray<IPLSphere> ProbeSpheres;
+			PhononProbeVolumeHandle->PlaceProbes(PhononScene, GenerateProbesProgressCallback, ProbeSpheres);
+			PhononProbeVolumeHandle->BakedDataInfo.Empty();
+
+			// Clean up
+			iplDestroyScene(&SceneCopy);
+
+			if (ComputeDevice)
+			{
+				iplDestroyComputeDevice(&ComputeDevice);
+			}
+
+			// Update probe component with new probe locations
+			{
+				FScopeLock ScopeLock(&PhononProbeVolumeHandle->PhononProbeComponent->ProbeLocationsCriticalSection);
+				auto& ProbeLocations = PhononProbeVolumeHandle->GetPhononProbeComponent()->ProbeLocations;
+				ProbeLocations.Empty();
+				ProbeLocations.SetNumUninitialized(ProbeSpheres.Num());
+				for (auto i = 0; i < ProbeSpheres.Num(); ++i)
+				{
+					ProbeLocations[i] = SteamAudio::PhononToUnrealFVector(SteamAudio::FVectorFromIPLVector3(ProbeSpheres[i].center));
+				}
+			}
+			
+			// Notify UI that we're done
+			GGenerateProbesTickable->QueueWorkItem(FWorkItem([](FText& DisplayText) {
+				DisplayText = NSLOCTEXT("SteamAudio", "ProbePlacementComplete", "Probe placement complete."); 
+			}, SNotificationItem::CS_Success, true));
+		});
+
+		return FReply::Handled();
+	}
+}
diff --git a/Engine/Plugins/Runtime/Steam/SteamAudio/Source/SteamAudioEditor/Private/PhononProbeVolumeDetails.h b/Engine/Plugins/Runtime/Steam/SteamAudio/Source/SteamAudioEditor/Private/PhononProbeVolumeDetails.h
new file mode 100644
--- /dev/null
+++ b/Engine/Plugins/Runtime/Steam/SteamAudio/Source/SteamAudioEditor/Private/PhononProbeVolumeDetails.h
@@ -0,0 +1,32 @@
+//
+// Copyright (C) Valve Corporation. All rights reserved.
+//
+
+#pragma once
+
+#include "IDetailCustomization.h"
+#include "Input/Reply.h"
+
+class IDetailLayoutBuilder;
+class IDetailChildrenBuilder;
+class IPropertyHandle;
+class APhononProbeVolume;
+
+namespace SteamAudio
+{
+	class FPhononProbeVolumeDetails : public IDetailCustomization
+	{
+	public:
+		static TSharedRef<IDetailCustomization> MakeInstance();
+
+	private:
+		virtual void CustomizeDetails(IDetailLayoutBuilder& DetailLayout) override;
+
+		void OnGenerateBakedDataInfo(TSharedRef<IPropertyHandle> PropertyHandle, int32 ArrayIndex, IDetailChildrenBuilder& ChildrenBuilder);
+		void OnClearBakedDataClicked(const int32 ArrayIndex);
+		FText GetTotalDataSize();
+
+		FReply OnGenerateProbes();
+		TWeakObjectPtr<APhononProbeVolume> PhononProbeVolume;
+	};
+}
\ No newline at end of file
diff --git a/Engine/Plugins/Runtime/Steam/SteamAudio/Source/SteamAudioEditor/Private/PhononReverbSettingsFactory.cpp b/Engine/Plugins/Runtime/Steam/SteamAudio/Source/SteamAudioEditor/Private/PhononReverbSettingsFactory.cpp
new file mode 100644
--- /dev/null
+++ b/Engine/Plugins/Runtime/Steam/SteamAudio/Source/SteamAudioEditor/Private/PhononReverbSettingsFactory.cpp
@@ -0,0 +1,64 @@
+//
+// Copyright (C) Valve Corporation. All rights reserved.
+//
+
+#include "PhononReverbSettingsFactory.h"
+
+#include "AudioAnalytics.h"
+#include "PhononReverbSourceSettings.h"
+#include "UObject/ObjectMacros.h"
+#include "UObject/Object.h"
+
+namespace SteamAudio
+{
+	FText FAssetTypeActions_PhononReverbSettings::GetName() const
+	{
+		return NSLOCTEXT("SteamAudio", "AssetTypeActions_PhononReverbSettings", "Phonon Source Reverb Settings");
+	}
+
+	FColor FAssetTypeActions_PhononReverbSettings::GetTypeColor() const
+	{
+		return FColor(245, 195, 101);
+	}
+
+	UClass* FAssetTypeActions_PhononReverbSettings::GetSupportedClass() const
+	{
+		return UPhononReverbSourceSettings::StaticClass();
+	}
+	
+	uint32 FAssetTypeActions_PhononReverbSettings::GetCategories()
+	{
+		return EAssetTypeCategories::Sounds;
+	}
+
+	const TArray<FText>& FAssetTypeActions_PhononReverbSettings::GetSubMenus() const
+	{
+		static const TArray<FText> PhononSubMenus
+		{
+			NSLOCTEXT("SteamAudio", "AssetPhononSubMenu", "Phonon")
+		};
+		return PhononSubMenus;
+	}
+}
+
+UPhononReverbSettingsFactory::UPhononReverbSettingsFactory(const FObjectInitializer& ObjectInitializer)
+	: Super(ObjectInitializer)
+{
+	SupportedClass = UPhononReverbSourceSettings::StaticClass();
+
+	bCreateNew = true;
+	bEditorImport = false;
+	bEditAfterNew = true;
+}
+
+UObject* UPhononReverbSettingsFactory::FactoryCreateNew(UClass* Class, UObject* InParent, FName InName, EObjectFlags Flags, UObject* Context,
+	FFeedbackContext* Warn)
+{
+	Audio::Analytics::RecordEvent_Usage(TEXT("SteamAudio.PhononReverbSettingsCreated"));
+	return NewObject<UPhononReverbSourceSettings>(InParent, InName, Flags);
+}
+
+uint32 UPhononReverbSettingsFactory::GetMenuCategories() const
+{
+	return EAssetTypeCategories::Sounds;
+}
\ No newline at end of file
diff --git a/Engine/Plugins/Runtime/Steam/SteamAudio/Source/SteamAudioEditor/Private/PhononReverbSettingsFactory.h b/Engine/Plugins/Runtime/Steam/SteamAudio/Source/SteamAudioEditor/Private/PhononReverbSettingsFactory.h
new file mode 100644
--- /dev/null
+++ b/Engine/Plugins/Runtime/Steam/SteamAudio/Source/SteamAudioEditor/Private/PhononReverbSettingsFactory.h
@@ -0,0 +1,36 @@
+//
+// Copyright (C) Valve Corporation. All rights reserved.
+//
+
+#pragma once
+
+#include "IAudioExtensionPlugin.h"
+#include "AssetToolsModule.h"
+#include "Factories/Factory.h"
+#include "AssetTypeActions_Base.h"
+#include "PhononReverbSettingsFactory.generated.h"
+
+namespace SteamAudio
+{
+	class FAssetTypeActions_PhononReverbSettings : public FAssetTypeActions_Base
+	{
+	public:
+		virtual FText GetName() const override;
+		virtual FColor GetTypeColor() const override;
+		virtual UClass* GetSupportedClass() const override;
+		virtual uint32 GetCategories() override;
+
+		virtual const TArray<FText>& GetSubMenus() const override;
+	};
+}
+
+UCLASS(MinimalAPI, hidecategories = Object)
+class UPhononReverbSettingsFactory : public UFactory
+{
+	GENERATED_UCLASS_BODY()
+
+	virtual UObject* FactoryCreateNew(UClass* Class, UObject* InParent, FName Name, EObjectFlags Flags, UObject* Context,
+		FFeedbackContext* Warn) override;
+
+	virtual uint32 GetMenuCategories() const override;
+};
diff --git a/Engine/Plugins/Runtime/Steam/SteamAudio/Source/SteamAudioEditor/Private/PhononSourceComponentDetails.cpp b/Engine/Plugins/Runtime/Steam/SteamAudio/Source/SteamAudioEditor/Private/PhononSourceComponentDetails.cpp
new file mode 100644
--- /dev/null
+++ b/Engine/Plugins/Runtime/Steam/SteamAudio/Source/SteamAudioEditor/Private/PhononSourceComponentDetails.cpp
@@ -0,0 +1,164 @@
+//
+// Copyright (C) Valve Corporation. All rights reserved.
+//
+
+#include "PhononSourceComponentDetails.h"
+
+#include "PhononCommon.h"
+#include "PhononSourceComponent.h"
+#include "SteamAudioEditorModule.h"
+#include "IndirectBaker.h"
+
+#include "Components/AudioComponent.h"
+#include "DetailLayoutBuilder.h"
+#include "DetailCategoryBuilder.h"
+#include "DetailWidgetRow.h"
+#include "IDetailsView.h"
+#include "IDetailCustomization.h"
+#include "IDetailChildrenBuilder.h"
+#include "LevelEditorViewport.h"
+#include "EngineUtils.h"
+#include "UObject/UObjectIterator.h"
+#include "Widgets/Input/SButton.h"
+#include "Widgets/Text/STextBlock.h"
+#include "Editor.h"
+
+namespace SteamAudio
+{
+	//==============================================================================================================================================
+	// FPhononSourceComponentDetails
+	//==============================================================================================================================================
+
+	TSharedRef<IDetailCustomization> FPhononSourceComponentDetails::MakeInstance()
+	{
+		return MakeShareable(new FPhononSourceComponentDetails());
+	}
+
+	void FPhononSourceComponentDetails::CustomizeDetails(IDetailLayoutBuilder& DetailLayout)
+	{
+		const TArray<TWeakObjectPtr<UObject>>& SelectedObjects = DetailLayout.GetSelectedObjects();
+
+		for (int32 ObjectIndex = 0; ObjectIndex < SelectedObjects.Num(); ++ObjectIndex)
+		{
+			const TWeakObjectPtr<UObject>& CurrentObject = SelectedObjects[ObjectIndex];
+			if (CurrentObject.IsValid())
+			{
+				auto SelectedPhononSourceComponent = Cast<UPhononSourceComponent>(CurrentObject.Get());
+				if (SelectedPhononSourceComponent)
+				{
+					PhononSourceComponent = SelectedPhononSourceComponent;
+					break;
+				}
+			}
+		}
+
+		DetailLayout.EditCategory("Baking").AddProperty(GET_MEMBER_NAME_CHECKED(UPhononSourceComponent, UniqueIdentifier));
+		DetailLayout.EditCategory("Baking").AddProperty(GET_MEMBER_NAME_CHECKED(UPhononSourceComponent, BakingRadius));
+		DetailLayout.EditCategory("Baking").AddCustomRow(NSLOCTEXT("SteamAudio", "Bake Propagation", "Bake Propagation"))
+			.NameContent()
+			[
+				SNullWidget::NullWidget
+			]
+			.ValueContent()
+			[
+				SNew(SHorizontalBox)
+				+ SHorizontalBox::Slot()
+				.AutoWidth()
+				[
+					SNew(SButton)
+					.ContentPadding(2)
+					.VAlign(VAlign_Center)
+					.HAlign(HAlign_Center)
+					.IsEnabled(this, &FPhononSourceComponentDetails::IsBakeEnabled)
+					.OnClicked(this, &FPhononSourceComponentDetails::OnBakePropagation)
+					[
+						SNew(STextBlock)
+						.Text(NSLOCTEXT("SteamAudio", "BakePropagation", "Bake Propagation"))
+						.Font(IDetailLayoutBuilder::GetDetailFont())
+					]
+				]
+			];
+	}
+
+	FReply FPhononSourceComponentDetails::OnBakePropagation()
+	{
+		TArray<UPhononSourceComponent*> SourceComponents;
+		
+		// Retrieve all phonon source components and ensure it is the correct component instance in the persistent level
+		for (TObjectIterator<UPhononSourceComponent> PhononSourceObj; PhononSourceObj; ++PhononSourceObj)
+		{
+			UPhononSourceComponent *PhononComponent = *PhononSourceObj;
+			if ((PhononComponent && PhononComponent->IsValidLowLevelFast())
+				&& (PhononComponent->UniqueIdentifier == PhononSourceComponent.Get()->UniqueIdentifier))
+			{
+				if (PhononComponent->GetOwner())
+				{
+					// Debug only - get all audio components and print out their ids
+					//for (auto ActorComp : PhononComponent->GetOwner()->GetComponentsByClass(UAudioComponent::StaticClass()))
+					//{
+					//	UAudioComponent* AudioComponent = Cast<UAudioComponent>(ActorComp);
+
+					//	if (AudioComponent)
+					//	{
+					//		UE_LOG(LogTemp, Warning, TEXT("Found Audio Component [%s] with UserID [%s]"), *AudioComponent->GetFullName(), *AudioComponent->AudioComponentUserID.ToString());
+					//	}
+					//}
+					//
+					
+					if (PhononComponent->GetOwner()->GetFullName().Contains(TEXT("Transient.World")))
+					{
+						// Go through all audio components that the owner of this phonon source actor has
+						TArray<UActorComponent*> Components;
+						PhononComponent->GetOwner()->GetComponents(UAudioComponent::StaticClass(), Components);
+
+						for (auto ActorComp : Components)
+						{
+							if (ActorComp && ActorComp->IsValidLowLevel())
+							{
+								UAudioComponent* AudioComponent = Cast<UAudioComponent>(ActorComp);
+
+								if (AudioComponent 
+									&& AudioComponent->IsValidLowLevelFast()
+									&& AudioComponent->AudioComponentUserID != PhononSourceComponent->UniqueIdentifier
+									)
+								{
+									// Apply the phonon source id to this audio component
+									AudioComponent->AudioComponentUserID = PhononSourceComponent->UniqueIdentifier;
+									//UE_LOG(LogTemp, Warning, TEXT("Audio Component [%s] UserID is now set to [%s]"), *AudioComponent->GetFullName(), *AudioComponent->AudioComponentUserID.ToString());
+
+									// Check if there's a phonon source child for this audio component
+									for (auto ChildComp : AudioComponent->GetAttachChildren())
+									{
+										auto* ChildPhononComponent = Cast<UPhononSourceComponent>(ChildComp);
+
+										// If there is one, then apply it's child phonon source id, overriding and phonon source id already applied to it
+										if (ChildPhononComponent && ChildPhononComponent->IsValidLowLevelFast())
+										{
+											AudioComponent->AudioComponentUserID = ChildPhononComponent->UniqueIdentifier;
+											//UE_LOG(LogTemp, Warning, TEXT("Audio Component [%s] UserID is now set to [%s]"), *AudioComponent->GetFullName(), *AudioComponent->AudioComponentUserID.ToString());
+											break;
+										}
+									}
+								}
+							}
+						}
+					}
+					else
+					{
+						SourceComponents.Add(PhononComponent);
+					}
+				}
+			}
+		}
+		
+		Bake(SourceComponents, false, nullptr);
+		
+		return FReply::Handled();
+	}
+
+	bool FPhononSourceComponentDetails::IsBakeEnabled() const
+	{
+		const bool bIsBaking = GIsBaking.Load();
+		return !(PhononSourceComponent->UniqueIdentifier.IsNone() || bIsBaking);
+	}
+}
diff --git a/Engine/Plugins/Runtime/Steam/SteamAudio/Source/SteamAudioEditor/Private/PhononSourceComponentDetails.h b/Engine/Plugins/Runtime/Steam/SteamAudio/Source/SteamAudioEditor/Private/PhononSourceComponentDetails.h
new file mode 100644
--- /dev/null
+++ b/Engine/Plugins/Runtime/Steam/SteamAudio/Source/SteamAudioEditor/Private/PhononSourceComponentDetails.h
@@ -0,0 +1,28 @@
+//
+// Copyright (C) Valve Corporation. All rights reserved.
+//
+
+#pragma once
+
+#include "IDetailCustomization.h"
+#include "Input/Reply.h"
+
+class UPhononSourceComponent;
+class IDetailLayoutBuilder;
+
+namespace SteamAudio
+{
+	class FPhononSourceComponentDetails : public IDetailCustomization
+	{
+	public:
+		static TSharedRef<IDetailCustomization> MakeInstance();
+
+	private:
+		virtual void CustomizeDetails(IDetailLayoutBuilder& DetailLayout) override;
+
+		FReply OnBakePropagation();
+		bool IsBakeEnabled() const;
+
+		TWeakObjectPtr<UPhononSourceComponent> PhononSourceComponent;
+	};
+}
diff --git a/Engine/Plugins/Runtime/Steam/SteamAudio/Source/SteamAudioEditor/Private/PhononSourceComponentVisualizer.cpp b/Engine/Plugins/Runtime/Steam/SteamAudio/Source/SteamAudioEditor/Private/PhononSourceComponentVisualizer.cpp
new file mode 100644
--- /dev/null
+++ b/Engine/Plugins/Runtime/Steam/SteamAudio/Source/SteamAudioEditor/Private/PhononSourceComponentVisualizer.cpp
@@ -0,0 +1,23 @@
+//
+// Copyright (C) Valve Corporation. All rights reserved.
+//
+
+#include "PhononSourceComponentVisualizer.h"
+#include "PhononSourceComponent.h"
+#include "SceneView.h"
+#include "SceneManagement.h"
+
+namespace SteamAudio
+{
+	void FPhononSourceComponentVisualizer::DrawVisualization(const UActorComponent* Component, const FSceneView* View, FPrimitiveDrawInterface* PDI)
+	{
+		const UPhononSourceComponent* AttenuatedComponent = Cast<const UPhononSourceComponent>(Component);
+		if (AttenuatedComponent != nullptr)
+		{
+			const FColor OuterRadiusColor(0, 153, 255);
+			auto Translation = AttenuatedComponent->GetComponentTransform().GetTranslation();
+
+			DrawWireSphereAutoSides(PDI, Translation, OuterRadiusColor, AttenuatedComponent->BakingRadius, SDPG_World);
+		}
+	}
+}
\ No newline at end of file
diff --git a/Engine/Plugins/Runtime/Steam/SteamAudio/Source/SteamAudioEditor/Private/PhononSourceComponentVisualizer.h b/Engine/Plugins/Runtime/Steam/SteamAudio/Source/SteamAudioEditor/Private/PhononSourceComponentVisualizer.h
new file mode 100644
--- /dev/null
+++ b/Engine/Plugins/Runtime/Steam/SteamAudio/Source/SteamAudioEditor/Private/PhononSourceComponentVisualizer.h
@@ -0,0 +1,19 @@
+//
+// Copyright (C) Valve Corporation. All rights reserved.
+//
+
+#pragma once
+
+#include "ComponentVisualizer.h"
+
+class FPrimitiveDrawInterface;
+class FSceneView;
+
+namespace SteamAudio
+{
+	class FPhononSourceComponentVisualizer : public FComponentVisualizer
+	{
+	public:
+		virtual void DrawVisualization(const UActorComponent* Component, const FSceneView* View, FPrimitiveDrawInterface* PDI) override;
+	};
+}
diff --git a/Engine/Plugins/Runtime/Steam/SteamAudio/Source/SteamAudioEditor/Private/PhononSpatializationSettingsFactory.cpp b/Engine/Plugins/Runtime/Steam/SteamAudio/Source/SteamAudioEditor/Private/PhononSpatializationSettingsFactory.cpp
new file mode 100644
--- /dev/null
+++ b/Engine/Plugins/Runtime/Steam/SteamAudio/Source/SteamAudioEditor/Private/PhononSpatializationSettingsFactory.cpp
@@ -0,0 +1,64 @@
+//
+// Copyright (C) Valve Corporation. All rights reserved.
+//
+
+#include "PhononSpatializationSettingsFactory.h"
+
+#include "AudioAnalytics.h"
+#include "PhononSpatializationSourceSettings.h"
+#include "UObject/ObjectMacros.h"
+#include "UObject/Object.h"
+
+namespace SteamAudio
+{
+	FText FAssetTypeActions_PhononSpatializationSettings::GetName() const
+	{
+		return NSLOCTEXT("SteamAudio", "AssetTypeActions_PhononSpatializationSettings", "Phonon Source Spatialization Settings");
+	}
+
+	FColor FAssetTypeActions_PhononSpatializationSettings::GetTypeColor() const
+	{
+		return FColor(245, 195, 101);
+	}
+
+	UClass* FAssetTypeActions_PhononSpatializationSettings::GetSupportedClass() const
+	{
+		return UPhononSpatializationSourceSettings::StaticClass();
+	}
+
+	uint32 FAssetTypeActions_PhononSpatializationSettings::GetCategories()
+	{
+		return EAssetTypeCategories::Sounds;
+	}
+
+	const TArray<FText>& FAssetTypeActions_PhononSpatializationSettings::GetSubMenus() const
+	{
+		static const TArray<FText> PhononSubMenus
+		{
+			NSLOCTEXT("SteamAudio", "AssetPhononSubMenu", "Phonon")
+		};
+		return PhononSubMenus;
+	}
+}
+
+UPhononSpatializationSettingsFactory::UPhononSpatializationSettingsFactory(const FObjectInitializer& ObjectInitializer)
+	: Super(ObjectInitializer)
+{
+	SupportedClass = UPhononSpatializationSourceSettings::StaticClass();
+
+	bCreateNew = true;
+	bEditorImport = false;
+	bEditAfterNew = true;
+}
+
+UObject* UPhononSpatializationSettingsFactory::FactoryCreateNew(UClass* Class, UObject* InParent, FName InName, EObjectFlags Flags,
+	UObject* Context, FFeedbackContext* Warn)
+{
+	Audio::Analytics::RecordEvent_Usage(TEXT("SteamAudio.PhononSpatializationSettingsCreated"));
+	return NewObject<UPhononSpatializationSourceSettings>(InParent, InName, Flags);
+}
+
+uint32 UPhononSpatializationSettingsFactory::GetMenuCategories() const
+{
+	return EAssetTypeCategories::Sounds;
+}
\ No newline at end of file
diff --git a/Engine/Plugins/Runtime/Steam/SteamAudio/Source/SteamAudioEditor/Private/PhononSpatializationSettingsFactory.h b/Engine/Plugins/Runtime/Steam/SteamAudio/Source/SteamAudioEditor/Private/PhononSpatializationSettingsFactory.h
new file mode 100644
--- /dev/null
+++ b/Engine/Plugins/Runtime/Steam/SteamAudio/Source/SteamAudioEditor/Private/PhononSpatializationSettingsFactory.h
@@ -0,0 +1,36 @@
+//
+// Copyright (C) Valve Corporation. All rights reserved.
+//
+
+#pragma once
+
+#include "IAudioExtensionPlugin.h"
+#include "AssetToolsModule.h"
+#include "Factories/Factory.h"
+#include "AssetTypeActions_Base.h"
+#include "PhononSpatializationSettingsFactory.generated.h"
+
+namespace SteamAudio
+{
+	class FAssetTypeActions_PhononSpatializationSettings : public FAssetTypeActions_Base
+	{
+	public:
+		virtual FText GetName() const override;
+		virtual FColor GetTypeColor() const override;
+		virtual UClass* GetSupportedClass() const override;
+		virtual uint32 GetCategories() override;
+
+		virtual const TArray<FText>& GetSubMenus() const override;
+	};
+}
+
+UCLASS(MinimalAPI, hidecategories = Object)
+class UPhononSpatializationSettingsFactory : public UFactory
+{
+	GENERATED_UCLASS_BODY()
+
+	virtual UObject* FactoryCreateNew(UClass* Class, UObject* InParent, FName Name, EObjectFlags Flags, UObject* Context,
+		FFeedbackContext* Warn) override;
+
+	virtual uint32 GetMenuCategories() const override;
+};
diff --git a/Engine/Plugins/Runtime/Steam/SteamAudio/Source/SteamAudioEditor/Private/SteamAudioEdMode.cpp b/Engine/Plugins/Runtime/Steam/SteamAudio/Source/SteamAudioEditor/Private/SteamAudioEdMode.cpp
new file mode 100644
--- /dev/null
+++ b/Engine/Plugins/Runtime/Steam/SteamAudio/Source/SteamAudioEditor/Private/SteamAudioEdMode.cpp
@@ -0,0 +1,42 @@
+//
+// Copyright (C) Valve Corporation. All rights reserved.
+//
+
+#include "SteamAudioEdMode.h"
+#include "SteamAudioEdModeToolkit.h"
+#include "SteamAudioEditorModule.h"
+
+#include "Toolkits/ToolkitManager.h"
+#include "EditorModeManager.h"
+
+namespace SteamAudio
+{
+	const FEditorModeID FSteamAudioEdMode::EM_SteamAudio = TEXT("EM_SteamAudio");
+
+	void FSteamAudioEdMode::Enter()
+	{
+		FEdMode::Enter();
+
+		if (!Toolkit.IsValid() && UsesToolkits())
+		{
+			Toolkit = MakeShareable(new FSteamAudioEdModeToolkit);
+			Toolkit->Init(Owner->GetToolkitHost());
+		}
+	}
+
+	void FSteamAudioEdMode::Exit()
+	{
+		if (Toolkit.IsValid())
+		{
+			FToolkitManager::Get().CloseToolkit(Toolkit.ToSharedRef());
+			Toolkit.Reset();
+		}
+
+		FEdMode::Exit();
+	}
+
+	bool FSteamAudioEdMode::UsesToolkits() const
+	{
+		return true;
+	}
+}
diff --git a/Engine/Plugins/Runtime/Steam/SteamAudio/Source/SteamAudioEditor/Private/SteamAudioEdMode.h b/Engine/Plugins/Runtime/Steam/SteamAudio/Source/SteamAudioEditor/Private/SteamAudioEdMode.h
new file mode 100644
--- /dev/null
+++ b/Engine/Plugins/Runtime/Steam/SteamAudio/Source/SteamAudioEditor/Private/SteamAudioEdMode.h
@@ -0,0 +1,22 @@
+//
+// Copyright (C) Valve Corporation. All rights reserved.
+//
+
+#pragma once
+
+#include "EdMode.h"
+#include "EditorModes.h"
+#include "EditorModeTools.h"
+
+namespace SteamAudio
+{
+	class FSteamAudioEdMode : public FEdMode
+	{
+	public:
+		static const FEditorModeID EM_SteamAudio;
+
+		virtual void Enter() override;
+		virtual void Exit() override;
+		bool UsesToolkits() const override;
+	};
+}
diff --git a/Engine/Plugins/Runtime/Steam/SteamAudio/Source/SteamAudioEditor/Private/SteamAudioEdModeToolkit.cpp b/Engine/Plugins/Runtime/Steam/SteamAudio/Source/SteamAudioEditor/Private/SteamAudioEdModeToolkit.cpp
new file mode 100644
--- /dev/null
+++ b/Engine/Plugins/Runtime/Steam/SteamAudio/Source/SteamAudioEditor/Private/SteamAudioEdModeToolkit.cpp
@@ -0,0 +1,421 @@
+//
+// Copyright (C) Valve Corporation. All rights reserved.
+//
+
+#include "SteamAudioEdModeToolkit.h"
+#include "SteamAudioEdMode.h"
+#include "SteamAudioEditorModule.h"
+#include "TickableNotification.h"
+#include "PhononScene.h"
+#include "PhononCommon.h"
+
+#include "Widgets/Layout/SScrollBox.h"
+#include "Widgets/Layout/SBox.h"
+#include "Widgets/Input/SButton.h"
+#include "Widgets/Text/STextBlock.h"
+#include "DetailLayoutBuilder.h"
+#include "DetailCategoryBuilder.h"
+#include "IDetailPropertyRow.h"
+#include "DetailWidgetRow.h"
+#include "IDetailGroup.h"
+#include "Editor.h"
+#include "LevelEditorViewport.h"
+#include "Settings/LevelEditorPlaySettings.h"
+#include "Settings/LevelEditorViewportSettings.h"
+#include "EditorModeManager.h"
+#include "Async/Async.h"
+
+namespace SteamAudio
+{
+	static TSharedPtr<FTickableNotification> GEdModeTickable = MakeShareable(new FTickableNotification());
+
+	FSteamAudioEdModeToolkit::FSteamAudioEdModeToolkit()
+	{
+		FSteamAudioEditorModule* Module = &FModuleManager::GetModuleChecked<FSteamAudioEditorModule>("SteamAudioEditor");
+		if (Module != nullptr)
+		{
+			UWorld* World = GEditor->GetLevelViewportClients()[0]->GetWorld();
+			FPhononSceneInfo PhononSceneInfo;
+			LoadSceneInfoFromDisk(World, PhononSceneInfo);
+			Module->SetCurrentPhononSceneInfo(PhononSceneInfo);
+		}
+	}
+
+	FSteamAudioEdModeToolkit::~FSteamAudioEdModeToolkit()
+	{
+	}
+
+	void FSteamAudioEdModeToolkit::Init(const TSharedPtr<IToolkitHost>& InitToolkitHost)
+	{
+		SAssignNew(ToolkitWidget, SScrollBox)
+		+ SScrollBox::Slot()
+		.Padding(5)
+		.HAlign(HAlign_Left)
+		[
+			SNew(SVerticalBox)
+			+ SVerticalBox::Slot()
+			[
+				SNew(SButton)
+				.ContentPadding(2)
+				.VAlign(VAlign_Center)
+				.HAlign(HAlign_Left)
+				.IsEnabled(this, &FSteamAudioEdModeToolkit::IsAddAllEnabled)
+				.ToolTipText(NSLOCTEXT("SteamAudio", "AddAllTooltip", "Adds Phonon Geometry components (with default settings) to all Static Mesh actors that do not already have one."))
+				.OnClicked(this, &FSteamAudioEdModeToolkit::OnAddAll)
+				[
+					SNew(STextBlock)
+					.Text(NSLOCTEXT("SteamAudio", "AddAll", "Add All"))
+					.Font(IDetailLayoutBuilder::GetDetailFont())
+				]
+			]
+			+ SVerticalBox::Slot()
+			[
+				SNew(SButton)
+				.ContentPadding(2)
+				.VAlign(VAlign_Center)
+				.HAlign(HAlign_Left)
+				.IsEnabled(this, &FSteamAudioEdModeToolkit::IsRemoveAllEnabled)
+				.ToolTipText(NSLOCTEXT("SteamAudio", "RemoveAllTooltip", "Removes all Phonon Geometry components from Static Mesh actors."))
+				.OnClicked(this, &FSteamAudioEdModeToolkit::OnRemoveAll)
+				[
+					SNew(STextBlock)
+					.Text(NSLOCTEXT("SteamAudio", "RemoveAll", "Remove All"))
+					.Font(IDetailLayoutBuilder::GetDetailFont())
+				]
+			]
+			+ SVerticalBox::Slot()
+			[
+				SNew(SButton)
+				.ContentPadding(2)
+				.VAlign(VAlign_Center)
+				.HAlign(HAlign_Left)
+				.IsEnabled(this, &FSteamAudioEdModeToolkit::IsExportObjEnabled)
+				.ToolTipText(NSLOCTEXT("SteamAudio", "ExportObjTooltip", "Exports all Phonon Geometry to an OBJ file in /Content. Used to debug scene export."))
+				.OnClicked(this, &FSteamAudioEdModeToolkit::OnExportObj)
+				[
+					SNew(STextBlock)
+					.Text(NSLOCTEXT("SteamAudio", "ExportOBJ", "Export OBJ"))
+					.Font(IDetailLayoutBuilder::GetDetailFont())
+				]
+			]
+			+ SVerticalBox::Slot()
+			[
+				SNew(SButton)
+				.ContentPadding(2)
+				.VAlign(VAlign_Center)
+				.HAlign(HAlign_Left)
+				.IsEnabled(this, &FSteamAudioEdModeToolkit::IsExportSceneEnabled)
+				.ToolTipText(NSLOCTEXT("SteamAudio", "ExportSceneTooltip", "Exports all Phonon Geometry for use by the Steam Audio system."))
+				.OnClicked(this, &FSteamAudioEdModeToolkit::OnExportScene)
+				[
+					SNew(STextBlock)
+					.Text(NSLOCTEXT("SteamAudio", "ExportScene", "Export Scene"))
+					.Font(IDetailLayoutBuilder::GetDetailFont())
+				]
+			]
+			+ SVerticalBox::Slot()
+			[
+				SNew(SHorizontalBox)
+				+ SHorizontalBox::Slot()
+				.AutoWidth()
+				.Padding(4)
+				[
+					SNew(STextBlock)
+					.Text(NSLOCTEXT("SteamAudio", "SceneTriangles", "Scene Triangles"))
+					.Font(IDetailLayoutBuilder::GetDetailFont())
+				]
+				+ SHorizontalBox::Slot()
+				.AutoWidth()
+				.Padding(4)
+				[
+					SNew(STextBlock)
+					.Text(this, &FSteamAudioEdModeToolkit::GetNumSceneTrianglesText)
+					.Font(IDetailLayoutBuilder::GetDetailFont())
+				]
+			]
+			+ SVerticalBox::Slot()
+			[
+				SNew(SHorizontalBox)
+				+ SHorizontalBox::Slot()
+				.AutoWidth()
+				.Padding(4)
+				[
+					SNew(STextBlock)
+					.Text(NSLOCTEXT("SteamAudio", "DynamicSceneTriangles", "Dynamic Scene Triangles"))
+					.Font(IDetailLayoutBuilder::GetDetailFont())
+				]
+				+ SHorizontalBox::Slot()
+				.AutoWidth()
+				.Padding(4)
+				[
+					SNew(STextBlock)
+					.Text(this, &FSteamAudioEdModeToolkit::GetNumDynSceneTrianglesText)
+					.Font(IDetailLayoutBuilder::GetDetailFont())
+				]
+			]
+			+ SVerticalBox::Slot()
+			[
+				SNew(SHorizontalBox)
+				+ SHorizontalBox::Slot()
+				.AutoWidth()
+				.Padding(4)
+				[
+					SNew(STextBlock)
+					.Text(NSLOCTEXT("SteamAudio", "SceneDataSize", "Scene Data Size"))
+					.Font(IDetailLayoutBuilder::GetDetailFont())
+				]
+				+ SHorizontalBox::Slot()
+				.AutoWidth()
+				.Padding(4)
+				[
+					SNew(STextBlock)
+					.Text(this, &FSteamAudioEdModeToolkit::GetSceneDataSizeText)
+					.Font(IDetailLayoutBuilder::GetDetailFont())
+				]
+			]
+			+ SVerticalBox::Slot()
+			[
+				SNew(SHorizontalBox)
+				+ SHorizontalBox::Slot()
+				.AutoWidth()
+				.Padding(4)
+				[
+					SNew(STextBlock)
+					.Text(NSLOCTEXT("SteamAudio", "DynamicSceneDataSize", "Dynamic Scene Data Size"))
+					.Font(IDetailLayoutBuilder::GetDetailFont())
+				]
+				+ SHorizontalBox::Slot()
+				.AutoWidth()
+				.Padding(4)
+				[
+					SNew(STextBlock)
+					.Text(this, &FSteamAudioEdModeToolkit::GetDynSceneDataSizeText)
+				.Font(IDetailLayoutBuilder::GetDetailFont())
+				]
+			]
+		];
+
+		FModeToolkit::Init(InitToolkitHost);
+	}
+
+	FName FSteamAudioEdModeToolkit::GetToolkitFName() const
+	{
+		return FName("SteamAudioEdMode");
+	}
+
+	FText FSteamAudioEdModeToolkit::GetBaseToolkitName() const
+	{
+		return NSLOCTEXT("SteamAudio", "ToolkitDisplayName", "Steam Audio Editor Mode Tool");
+	}
+
+	FEdMode* FSteamAudioEdModeToolkit::GetEditorMode() const
+	{
+		return GLevelEditorModeTools().GetActiveMode(FSteamAudioEdMode::EM_SteamAudio);
+	}
+
+	FReply FSteamAudioEdModeToolkit::OnAddAll()
+	{
+		// Display editor notification
+		GEdModeTickable->SetDisplayText(NSLOCTEXT("SteamAudio", "AddingPhononGeometryToAllStaticMeshActors", "Adding Phonon Geometry to all Static Mesh actors..."));
+		GEdModeTickable->CreateNotification();
+
+		AsyncTask(ENamedThreads::GameThread, [=]()
+		{
+			UWorld* World = GEditor->GetLevelViewportClients()[0]->GetWorld();
+			AddGeometryComponentsToStaticMeshes(World);
+
+			// Notify UI that we're done
+			GEdModeTickable->QueueWorkItem(FWorkItem([](FText& DisplayText) {
+				DisplayText = NSLOCTEXT("SteamAudio", "ComponentsAdded", "Components added.");
+			}, SNotificationItem::CS_Success, true));
+		});
+
+		return FReply::Handled();
+	}
+
+	bool FSteamAudioEdModeToolkit::IsAddAllEnabled() const
+	{
+		return true;
+	}
+
+	FReply FSteamAudioEdModeToolkit::OnRemoveAll()
+	{
+		// Display editor notification
+		GEdModeTickable->SetDisplayText(NSLOCTEXT("SteamAudio", "RemovingPhononGeometryFromAllStaticMeshActors", "Removing Phonon Geometry from all Static Mesh actors..."));
+		GEdModeTickable->CreateNotification();
+
+		AsyncTask(ENamedThreads::GameThread, [=]()
+		{
+			UWorld* World = GEditor->GetLevelViewportClients()[0]->GetWorld();
+			RemoveGeometryComponentsFromStaticMeshes(World);
+
+			// Notify UI that we're done
+			GEdModeTickable->QueueWorkItem(FWorkItem([](FText& DisplayText) {
+				DisplayText = NSLOCTEXT("SteamAudio", "ComponentsRemoved", "Components removed.");
+			}, SNotificationItem::CS_Success, true));
+		});
+
+		return FReply::Handled();
+	}
+
+	bool FSteamAudioEdModeToolkit::IsRemoveAllEnabled() const
+	{
+		return true;
+	}
+
+	static void FinalizeSceneProgressCallback(IPLfloat32 Progress)
+	{
+		FFormatNamedArguments Arguments;
+		Arguments.Add(TEXT("FinalizeSceneProgress"), FText::AsPercent(Progress));
+		GEdModeTickable->SetDisplayText(FText::Format(NSLOCTEXT("SteamAudio", "FinalizeSceneProgressFmt", "Finalizing scene ({FinalizeSceneProgress} complete)"), Arguments));
+	}
+
+	FReply FSteamAudioEdModeToolkit::OnExportObj()
+	{
+		// Display editor notification
+		GEdModeTickable->SetDisplayText(NSLOCTEXT("SteamAudio", "ExportingSceneOBJ", "Exporting scene OBJ..."));
+		GEdModeTickable->CreateNotification();
+
+		Async(EAsyncExecution::Thread, [&]()
+		{
+			// Export the scene
+			UWorld* World = GEditor->GetLevelViewportClients()[0]->GetWorld();
+			IPLhandle PhononScene = nullptr;
+			IPLhandle PhononStaticMesh = nullptr;
+			FPhononSceneInfo PhononSceneInfo;
+
+			if (!CreateScene(World, &PhononScene, &PhononStaticMesh, PhononSceneInfo.NumTriangles, PhononSceneInfo.NumDynTriangles, PhononSceneInfo.DynDataSize, true))
+			{
+				GEdModeTickable->QueueWorkItem(FWorkItem([](FText& DisplayText) {
+					DisplayText = NSLOCTEXT("SteamAudio", "ErrorExportingOBJ", "Error exporting OBJ.");
+				}, SNotificationItem::CS_Fail, true));
+				return;
+			}
+
+			// Write OBJ
+			iplSaveSceneAsObj(PhononScene, TCHAR_TO_ANSI(*(EditorOnlyPath + World->GetMapName() + ".obj")));
+
+			// Clean up Phonon structures
+			iplDestroyStaticMesh(&PhononStaticMesh);
+			iplDestroyScene(&PhononScene);
+
+			// Notify UI that we're done
+			GEdModeTickable->QueueWorkItem(FWorkItem([](FText& DisplayText) {
+				DisplayText = NSLOCTEXT("SteamAudio", "ExportSceneOBJComplete", "Export scene OBJ complete.");
+			}, SNotificationItem::CS_Success, true));
+		});
+
+		return FReply::Handled();
+	}
+
+	bool FSteamAudioEdModeToolkit::IsExportObjEnabled() const
+	{
+		return true;
+	}
+
+	FReply FSteamAudioEdModeToolkit::OnExportScene()
+	{
+		// Display editor notification
+		GEdModeTickable->SetDisplayText(NSLOCTEXT("SteamAudio", "ExportingScene", "Exporting scene..."));
+		GEdModeTickable->CreateNotification();
+
+		Async(EAsyncExecution::Thread, [&]()
+		{
+			// Create the scene
+			UWorld* World = GEditor->GetLevelViewportClients()[0]->GetWorld();
+			IPLhandle PhononScene = nullptr;
+			IPLhandle PhononStaticMesh = nullptr;
+			FPhononSceneInfo PhononSceneInfo;
+
+			if (!CreateScene(World, &PhononScene, &PhononStaticMesh, PhononSceneInfo.NumTriangles, PhononSceneInfo.NumDynTriangles, PhononSceneInfo.DynDataSize))
+			{
+				GEdModeTickable->QueueWorkItem(FWorkItem([](FText& DisplayText) {
+					DisplayText = NSLOCTEXT("SteamAudio", "ErrorExportingScene", "Error exporting scene.");
+				}, SNotificationItem::CS_Fail, true));
+				return;
+			}
+
+			// Finalize and save to disk
+			PhononSceneInfo.DataSize = iplSaveScene(PhononScene, nullptr);
+			bool SaveSceneSuccessful = SaveFinalizedSceneToDisk(World, PhononScene, PhononSceneInfo);
+
+			// Clean up Phonon structures
+			iplDestroyStaticMesh(&PhononStaticMesh);
+			iplDestroyScene(&PhononScene);
+
+			if (SaveSceneSuccessful)
+			{
+				FSteamAudioEditorModule* Module = &FModuleManager::GetModuleChecked<FSteamAudioEditorModule>("SteamAudioEditor");
+				if (Module != nullptr)
+				{
+					Module->SetCurrentPhononSceneInfo(PhononSceneInfo);
+				}
+
+				GEdModeTickable->QueueWorkItem(FWorkItem([](FText& DisplayText) {
+					DisplayText = NSLOCTEXT("SteamAudio", "ExportSceneComplete", "Export scene complete.");
+				}, SNotificationItem::CS_Success, true));
+			}
+			else
+			{
+				GEdModeTickable->QueueWorkItem(FWorkItem([](FText& DisplayText) {
+					DisplayText = NSLOCTEXT("SteamAudio", "ExportSceneFailed", "Export scene failed.");
+				}, SNotificationItem::CS_Fail, true));
+			}
+		});
+
+		return FReply::Handled();
+	}
+
+	bool FSteamAudioEdModeToolkit::IsExportSceneEnabled() const
+	{
+		return true;
+	}
+
+	FText FSteamAudioEdModeToolkit::GetNumSceneTrianglesText() const
+	{
+		FSteamAudioEditorModule* Module = &FModuleManager::GetModuleChecked<FSteamAudioEditorModule>("SteamAudioEditor");
+		if (Module != nullptr)
+		{
+			auto PhononSceneInfo = Module->GetCurrentPhononSceneInfo();
+			return FText::AsNumber(PhononSceneInfo.NumTriangles);
+		}
+
+		return FText::FromString("No .phononscene found");
+	}
+
+	FText FSteamAudioEdModeToolkit::GetNumDynSceneTrianglesText() const
+	{
+		FSteamAudioEditorModule* Module = &FModuleManager::GetModuleChecked<FSteamAudioEditorModule>("SteamAudioEditor");
+		if (Module != nullptr)
+		{
+			auto PhononSceneInfo = Module->GetCurrentPhononSceneInfo();
+			return FText::AsNumber(PhononSceneInfo.NumDynTriangles);
+		}
+
+		return FText::FromString("No .phononscene found");
+	}
+
+	FText FSteamAudioEdModeToolkit::GetSceneDataSizeText() const
+	{
+		FSteamAudioEditorModule* Module = &FModuleManager::GetModuleChecked<FSteamAudioEditorModule>("SteamAudioEditor");
+		if (Module != nullptr)
+		{
+			auto PhononSceneInfo = Module->GetCurrentPhononSceneInfo();
+			return FText::AsMemory(PhononSceneInfo.DataSize);
+		}
+
+		return FText::FromString("No .phononscene found");
+	}
+
+	FText FSteamAudioEdModeToolkit::GetDynSceneDataSizeText() const
+	{
+		FSteamAudioEditorModule* Module = &FModuleManager::GetModuleChecked<FSteamAudioEditorModule>("SteamAudioEditor");
+		if (Module != nullptr)
+		{
+			auto PhononSceneInfo = Module->GetCurrentPhononSceneInfo();
+			return FText::AsMemory(PhononSceneInfo.DynDataSize);
+		}
+
+		return FText::FromString("No .phononscene found");
+	}
+}
diff --git a/Engine/Plugins/Runtime/Steam/SteamAudio/Source/SteamAudioEditor/Private/SteamAudioEdModeToolkit.h b/Engine/Plugins/Runtime/Steam/SteamAudio/Source/SteamAudioEditor/Private/SteamAudioEdModeToolkit.h
new file mode 100644
--- /dev/null
+++ b/Engine/Plugins/Runtime/Steam/SteamAudio/Source/SteamAudioEditor/Private/SteamAudioEdModeToolkit.h
@@ -0,0 +1,53 @@
+//
+// Copyright (C) Valve Corporation. All rights reserved.
+//
+
+#pragma once
+
+#include "Editor/UnrealEd/Public/Toolkits/BaseToolkit.h"
+#include "IDetailsView.h"
+
+class FEdMode;
+class SWidget;
+class IToolkitHost;
+class FTabManager;
+class FRunnableThread;
+
+namespace SteamAudio
+{
+	class FSteamAudioEdModeToolkit : public FModeToolkit
+	{
+	public:
+		FSteamAudioEdModeToolkit();
+		~FSteamAudioEdModeToolkit();
+
+		/** Initializes the geometry mode toolkit */
+		virtual void Init(const TSharedPtr<IToolkitHost>& InitToolkitHost) override;
+
+		/** IToolkit interface */
+		virtual FName GetToolkitFName() const override;
+		virtual FText GetBaseToolkitName() const override;
+		virtual FEdMode* GetEditorMode() const override;
+		virtual TSharedPtr<SWidget> GetInlineContent() const override { return ToolkitWidget; }
+
+		FReply OnAddAll();
+		bool IsAddAllEnabled() const;
+
+		FReply OnRemoveAll();
+		bool IsRemoveAllEnabled() const;
+
+		FReply OnExportObj();
+		bool IsExportObjEnabled() const;
+
+		FReply OnExportScene();
+		bool IsExportSceneEnabled() const;
+
+		FText GetNumSceneTrianglesText() const;
+		FText GetNumDynSceneTrianglesText() const;
+		FText GetSceneDataSizeText() const;
+		FText GetDynSceneDataSizeText() const;
+
+	private:
+		TSharedPtr<SWidget> ToolkitWidget;
+	};
+}
\ No newline at end of file
diff --git a/Engine/Plugins/Runtime/Steam/SteamAudio/Source/SteamAudioEditor/Private/SteamAudioEditorModule.cpp b/Engine/Plugins/Runtime/Steam/SteamAudio/Source/SteamAudioEditor/Private/SteamAudioEditorModule.cpp
new file mode 100644
--- /dev/null
+++ b/Engine/Plugins/Runtime/Steam/SteamAudio/Source/SteamAudioEditor/Private/SteamAudioEditorModule.cpp
@@ -0,0 +1,240 @@
+//
+// Copyright (C) Valve Corporation. All rights reserved.
+//
+
+#include "SteamAudioEditorModule.h"
+
+#include "ISettingsModule.h"
+#include "Engine/World.h"
+#include "Styling/SlateStyleRegistry.h"
+#include "Kismet/GameplayStatics.h"
+#include "Async/Async.h"
+#include "UnrealEdGlobals.h"
+#include "Editor/UnrealEdEngine.h"
+#include "LevelEditorViewport.h"
+#include "LevelEditor.h"
+#include "PropertyEditorModule.h"
+#include "Styling/SlateStyle.h"
+#include "Interfaces/IPluginManager.h"
+#include "ClassIconFinder.h"
+#include "Editor.h"
+
+#include "BakeIndirectWindow.h"
+#include "TickableNotification.h"
+#include "PhononScene.h"
+#include "SteamAudioSettings.h"
+#include "IAssetTools.h"
+#include "PhononOcclusionSettingsFactory.h"
+#include "PhononProbeVolume.h"
+#include "PhononProbeVolumeDetails.h"
+#include "PhononProbeComponent.h"
+#include "PhononProbeComponentVisualizer.h"
+#include "PhononReverbSettingsFactory.h"
+#include "PhononScene.h"
+#include "PhononSourceComponent.h"
+#include "PhononSourceComponentDetails.h"
+#include "PhononSourceComponentVisualizer.h"
+#include "PhononSpatializationSettingsFactory.h"
+#include "PhononCommon.h"
+#include "SteamAudioEdMode.h"
+
+#include "Framework/MultiBox/MultiBoxBuilder.h"
+#include "AssetRegistry/AssetRegistryModule.h"
+#include "EditorModeManager.h"
+#include "HAL/PlatformFileManager.h"
+#include "GenericPlatform/GenericPlatformFile.h"
+#include "Misc/Paths.h"
+
+DEFINE_LOG_CATEGORY(LogSteamAudioEditor);
+
+IMPLEMENT_MODULE(SteamAudio::FSteamAudioEditorModule, SteamAudioEditor)
+
+
+namespace
+{
+	static const FName AssetToolsName = TEXT("AssetTools");
+
+	template <typename T>
+	void AddAssetAction(IAssetTools& AssetTools, TArray<TSharedPtr<FAssetTypeActions_Base>>& AssetActions)
+	{
+		TSharedPtr<T> AssetAction = MakeShared<T>();
+		TSharedPtr<FAssetTypeActions_Base> AssetActionBase = StaticCastSharedPtr<FAssetTypeActions_Base>(AssetAction);
+		AssetTools.RegisterAssetTypeActions(AssetAction.ToSharedRef());
+		AssetActions.Add(AssetActionBase);
+	}
+} // namespace <>
+
+namespace SteamAudio
+{
+	//==============================================================================================================================================
+	// FSteamAudioEditorModule
+	//==============================================================================================================================================
+
+	void FSteamAudioEditorModule::StartupModule()
+	{
+		// Ensure that Steam Audio has appropriate folders within which it can manage its data
+		IPlatformFile& PlatformFile = FPlatformFileManager::Get().GetPlatformFile();
+		
+		if (!PlatformFile.DirectoryExists(*BasePath))
+		{
+			PlatformFile.CreateDirectory(*BasePath);
+		}
+
+		if (!PlatformFile.DirectoryExists(*RuntimePath))
+		{
+			PlatformFile.CreateDirectory(*RuntimePath);
+		}
+
+		if (!PlatformFile.DirectoryExists(*DynamicRuntimePath))
+		{
+			PlatformFile.CreateDirectory(*DynamicRuntimePath);
+		}
+
+		if (!PlatformFile.DirectoryExists(*EditorOnlyPath))
+		{
+			PlatformFile.CreateDirectory(*EditorOnlyPath);
+		}
+
+		if (!PlatformFile.DirectoryExists(*DynamicEditorOnlyPath))
+		{
+			PlatformFile.CreateDirectory(*DynamicEditorOnlyPath);
+		}
+
+		// Register detail customizations
+		FPropertyEditorModule& PropertyModule = FModuleManager::LoadModuleChecked<FPropertyEditorModule>("PropertyEditor");
+
+		PropertyModule.RegisterCustomClassLayout("PhononProbeVolume", FOnGetDetailCustomizationInstance::CreateStatic(&FPhononProbeVolumeDetails::MakeInstance));
+		PropertyModule.RegisterCustomClassLayout("PhononSourceComponent",  FOnGetDetailCustomizationInstance::CreateStatic(&FPhononSourceComponentDetails::MakeInstance));
+
+		// Extend the toolbar build menu with custom actions
+		if (!IsRunningCommandlet())
+		{
+			FLevelEditorModule* LevelEditorModule = FModuleManager::LoadModulePtr<FLevelEditorModule>(TEXT("LevelEditor"));
+			if (LevelEditorModule)
+			{
+				auto BuildMenuExtender = FLevelEditorModule::FLevelEditorMenuExtender::CreateRaw(this, &FSteamAudioEditorModule::OnExtendLevelEditorBuildMenu);
+
+				LevelEditorModule->GetAllLevelEditorToolbarBuildMenuExtenders().Add(BuildMenuExtender);
+			}
+		}
+
+		// Register plugin settings
+		ISettingsModule* SettingsModule = FModuleManager::Get().GetModulePtr<ISettingsModule>("Settings");
+		if (SettingsModule)
+		{
+			SettingsModule->RegisterSettings("Project", "Plugins", "Steam Audio", NSLOCTEXT("SteamAudio", "SteamAudio", "Steam Audio"),
+				NSLOCTEXT("SteamAudio", "ConfigureSteamAudioSettings", "Configure Steam Audio settings"), GetMutableDefault<USteamAudioSettings>());
+		}
+
+		// Create bake indirect window
+		BakeIndirectWindow = MakeShareable(new FBakeIndirectWindow());
+
+		// Create and register custom slate style
+		FString SteamAudioContent = IPluginManager::Get().FindPlugin(TEXT("SteamAudio"))->GetBaseDir() + "/Content";
+		FVector2D Vec16 = FVector2D(16.0f, 16.0f);
+		FVector2D Vec40 = FVector2D(40.0f, 40.0f);
+		FVector2D Vec64 = FVector2D(64.0f, 64.0f);
+
+		SteamAudioStyleSet = MakeShareable(new FSlateStyleSet("SteamAudio"));
+		SteamAudioStyleSet->SetContentRoot(FPaths::EngineContentDir() / TEXT("Editor/Slate"));
+		SteamAudioStyleSet->SetCoreContentRoot(FPaths::EngineContentDir() / TEXT("Slate"));
+		SteamAudioStyleSet->Set("ClassIcon.PhononSourceComponent", new FSlateImageBrush(SteamAudioContent + "/S_PhononSource_16.png", Vec16));
+		SteamAudioStyleSet->Set("ClassIcon.PhononGeometryComponent", new FSlateImageBrush(SteamAudioContent + "/S_PhononGeometry_16.png", Vec16));
+		SteamAudioStyleSet->Set("ClassIcon.PhononMaterialComponent", new FSlateImageBrush(SteamAudioContent + "/S_PhononMaterial_16.png", Vec16));
+
+		SteamAudioStyleSet->Set("ClassIcon.PhononSpatializationSourceSettings", new FSlateImageBrush(SteamAudioContent + "/S_PhononSpatializationSourceSettings_16.png", Vec16));
+		SteamAudioStyleSet->Set("ClassThumbnail.PhononSpatializationSourceSettings", new FSlateImageBrush(SteamAudioContent + "/S_PhononSpatializationSourceSettings_64.png", Vec64));
+		
+		SteamAudioStyleSet->Set("ClassIcon.PhononOcclusionSourceSettings", new FSlateImageBrush(SteamAudioContent + "/S_PhononOcclusionSourceSettings_16.png", Vec16));
+		SteamAudioStyleSet->Set("ClassThumbnail.PhononOcclusionSourceSettings", new FSlateImageBrush(SteamAudioContent + "/S_PhononOcclusionSourceSettings_64.png", Vec64));
+		
+		SteamAudioStyleSet->Set("ClassIcon.PhononReverbSourceSettings", new FSlateImageBrush(SteamAudioContent + "/S_PhononReverbSourceSettings_16.png", Vec16));
+		SteamAudioStyleSet->Set("ClassThumbnail.PhononReverbSourceSettings", new FSlateImageBrush(SteamAudioContent + "/S_PhononReverbSourceSettings_64.png", Vec64));
+
+		SteamAudioStyleSet->Set("LevelEditor.SteamAudioMode", new FSlateImageBrush(SteamAudioContent + "/SteamAudio_EdMode_40.png", Vec40));
+		SteamAudioStyleSet->Set("LevelEditor.SteamAudioMode.Small", new FSlateImageBrush(SteamAudioContent + "/SteamAudio_EdMode_16.png", Vec16));
+
+		IAssetTools& AssetTools = FModuleManager::LoadModuleChecked<FAssetToolsModule>(AssetToolsName).Get();
+		AddAssetAction<FAssetTypeActions_PhononReverbSettings>(AssetTools, AssetActions);
+		AddAssetAction<FAssetTypeActions_PhononOcclusionSettings>(AssetTools, AssetActions);
+		AddAssetAction<FAssetTypeActions_PhononSpatializationSettings>(AssetTools, AssetActions);
+
+		FSlateStyleRegistry::RegisterSlateStyle(*SteamAudioStyleSet.Get());
+
+		// Register the ed mode
+		FEditorModeRegistry::Get().RegisterMode<FSteamAudioEdMode>(
+			FSteamAudioEdMode::EM_SteamAudio,
+			NSLOCTEXT("SteamAudio", "SteamAudioMode", "Steam Audio"),
+			FSlateIcon(SteamAudioStyleSet->GetStyleSetName(), "LevelEditor.SteamAudioMode", "LevelEditor.SteamAudioMode.Small"),
+			true);
+
+		// Register component visualizers
+		RegisterComponentVisualizer(UPhononSourceComponent::StaticClass()->GetFName(), MakeShareable(new FPhononSourceComponentVisualizer()));
+		RegisterComponentVisualizer(UPhononProbeComponent::StaticClass()->GetFName(), MakeShareable(new FPhononProbeComponentVisualizer()));
+	}
+
+	void FSteamAudioEditorModule::ShutdownModule()
+	{
+		// Unregister component visualizers
+		if (GUnrealEd)
+		{
+			for (FName& ClassName : RegisteredComponentClassNames)
+			{
+				GUnrealEd->UnregisterComponentVisualizer(ClassName);
+			}
+		}
+
+		AssetActions.Reset();
+
+		FSlateStyleRegistry::UnRegisterSlateStyle(*SteamAudioStyleSet.Get());
+	}
+
+	void FSteamAudioEditorModule::BakeIndirect()
+	{
+		BakeIndirectWindow->Invoke();
+	}
+
+	void FSteamAudioEditorModule::RegisterComponentVisualizer(const FName ComponentClassName, TSharedPtr<FComponentVisualizer> Visualizer)
+	{
+		if (GUnrealEd)
+		{
+			GUnrealEd->RegisterComponentVisualizer(ComponentClassName, Visualizer);
+		}
+
+		RegisteredComponentClassNames.Add(ComponentClassName);
+
+		if (Visualizer.IsValid())
+		{
+			Visualizer->OnRegister();
+		}
+	}
+
+	TSharedRef<FExtender> FSteamAudioEditorModule::OnExtendLevelEditorBuildMenu(const TSharedRef<FUICommandList> CommandList)
+	{
+		TSharedRef<FExtender> Extender(new FExtender());
+
+		Extender->AddMenuExtension("LevelEditorNavigation", EExtensionHook::After, nullptr,
+			FMenuExtensionDelegate::CreateRaw(this, &FSteamAudioEditorModule::CreateBuildMenu));
+
+		return Extender;
+	}
+
+	void FSteamAudioEditorModule::CreateBuildMenu(FMenuBuilder& Builder)
+	{
+		FUIAction ActionBakeIndirect(FExecuteAction::CreateRaw(this, &FSteamAudioEditorModule::BakeIndirect),
+			FCanExecuteAction::CreateRaw(this, &FSteamAudioEditorModule::IsReadyToBakeIndirect));
+
+		Builder.BeginSection("LevelEditorIR", NSLOCTEXT("SteamAudio", "SteamAudio", "Steam Audio"));
+
+		Builder.AddMenuEntry(NSLOCTEXT("SteamAudio", "BakeIndirectSound", "Bake Indirect Sound..."),
+			NSLOCTEXT("SteamAudio", "OpensIndirectBakingManager", "Opens indirect baking manager."), FSlateIcon(), ActionBakeIndirect,
+			NAME_None, EUserInterfaceActionType::Button);
+
+		Builder.EndSection();
+	}
+
+	bool FSteamAudioEditorModule::IsReadyToBakeIndirect() const
+	{
+		return true;
+	}
+}
diff --git a/Engine/Plugins/Runtime/Steam/SteamAudio/Source/SteamAudioEditor/Private/SteamAudioEditorModule.h b/Engine/Plugins/Runtime/Steam/SteamAudio/Source/SteamAudioEditor/Private/SteamAudioEditorModule.h
new file mode 100644
--- /dev/null
+++ b/Engine/Plugins/Runtime/Steam/SteamAudio/Source/SteamAudioEditor/Private/SteamAudioEditorModule.h
@@ -0,0 +1,47 @@
+//
+// Copyright (C) Valve Corporation. All rights reserved.
+//
+
+#pragma once
+
+#include "ISteamAudioEditorModule.h"
+
+#include "AssetTypeActions_Base.h"
+#include "Framework/MultiBox/MultiBoxExtender.h"
+#include "PhononScene.h"
+
+DECLARE_LOG_CATEGORY_EXTERN(LogSteamAudioEditor, Log, All);
+
+class FComponentVisualizer;
+class FSlateStyleSet;
+
+namespace SteamAudio
+{
+	class FBakeIndirectWindow;
+
+	class FSteamAudioEditorModule : public ISteamAudioEditorModule
+	{
+	public:
+		virtual void StartupModule() override;
+		virtual void ShutdownModule() override;
+
+		TSharedRef<FExtender> OnExtendLevelEditorBuildMenu(const TSharedRef<FUICommandList> CommandList);
+		void CreateBuildMenu(FMenuBuilder& Builder);
+
+		void BakeIndirect();
+		bool IsReadyToBakeIndirect() const;
+
+		void RegisterComponentVisualizer(const FName ComponentClassName, TSharedPtr<FComponentVisualizer> Visualizer);
+
+		void SetCurrentPhononSceneInfo(const FPhononSceneInfo& PhononSceneInfo) { CurrentPhononSceneInfo = PhononSceneInfo; }
+		const FPhononSceneInfo& GetCurrentPhononSceneInfo() const { return CurrentPhononSceneInfo; }
+
+	private:
+		TSharedPtr<FSlateStyleSet> SteamAudioStyleSet;
+		TSharedPtr<FBakeIndirectWindow> BakeIndirectWindow;
+		TArray<FName> RegisteredComponentClassNames;
+		FPhononSceneInfo CurrentPhononSceneInfo;
+
+		TArray<TSharedPtr<FAssetTypeActions_Base>> AssetActions;
+	};
+}
diff --git a/Engine/Plugins/Runtime/Steam/SteamAudio/Source/SteamAudioEditor/Private/TickableNotification.cpp b/Engine/Plugins/Runtime/Steam/SteamAudio/Source/SteamAudioEditor/Private/TickableNotification.cpp
new file mode 100644
--- /dev/null
+++ b/Engine/Plugins/Runtime/Steam/SteamAudio/Source/SteamAudioEditor/Private/TickableNotification.cpp
@@ -0,0 +1,117 @@
+//
+// Copyright (C) Valve Corporation. All rights reserved.
+//
+
+#include "TickableNotification.h"
+#include "Widgets/Notifications/SNotificationList.h"
+#include "Framework/Notifications/NotificationManager.h"
+#include "LevelEditor.h"
+#include "Misc/ScopeLock.h"
+
+namespace SteamAudio
+{
+	FWorkItem::FWorkItem()
+		: Task(nullptr)
+		, FinalState(SNotificationItem::CS_Success)
+		, bIsFinalItem(false)
+	{}
+
+	FWorkItem::FWorkItem(const TFunction<void(FText&)>& Task, const SNotificationItem::ECompletionState InFinalState, const bool bIsFinalItem)
+		: Task(Task)
+		, FinalState(InFinalState)
+		, bIsFinalItem(bIsFinalItem)
+	{}
+
+	FTickableNotification::FTickableNotification()
+		: bIsTicking(false)
+	{
+	}
+
+	void FTickableNotification::CreateNotification()
+	{
+		FNotificationInfo Info(DisplayText);
+		Info.bFireAndForget = false;
+		Info.FadeOutDuration = 4.0f;
+		Info.ExpireDuration = 0.0f;
+
+		NotificationPtr = FSlateNotificationManager::Get().AddNotification(Info);
+		NotificationPtr.Pin()->SetCompletionState(SNotificationItem::CS_Pending);
+
+		bIsTicking = true;
+	}
+	
+	void FTickableNotification::CreateNotificationWithCancel(const FSimpleDelegate& CancelDelegate)
+	{
+		FNotificationInfo Info(DisplayText);
+		Info.bFireAndForget = false;
+		Info.FadeOutDuration = 4.0f;
+		Info.ExpireDuration = 0.0f;
+		Info.ButtonDetails.Add(FNotificationButtonInfo(NSLOCTEXT("SteamAudio", "Cancel", "Cancel"), FText::GetEmpty(), CancelDelegate));
+
+		NotificationPtr = FSlateNotificationManager::Get().AddNotification(Info);
+		NotificationPtr.Pin()->SetCompletionState(SNotificationItem::CS_Pending);
+
+		bIsTicking = true;
+	}
+
+	void FTickableNotification::DestroyNotification(const SNotificationItem::ECompletionState InFinalState)
+	{
+		bIsTicking = false;
+		this->FinalState = InFinalState;
+	}
+
+	void FTickableNotification::SetDisplayText(const FText& InDisplayText)
+	{
+		FScopeLock Lock(&CriticalSection);
+
+		this->DisplayText = InDisplayText;
+	}
+
+	void FTickableNotification::QueueWorkItem(const FWorkItem& WorkItem)
+	{
+		FScopeLock Lock(&CriticalSection);
+
+		WorkQueue.Enqueue(WorkItem);
+	}
+
+	void FTickableNotification::NotifyDestruction()
+	{
+		if (!NotificationPtr.Pin().IsValid())
+		{
+			return;
+		}
+
+		NotificationPtr.Pin()->SetText(DisplayText);
+		NotificationPtr.Pin()->SetCompletionState(FinalState);
+		NotificationPtr.Pin()->ExpireAndFadeout();
+		NotificationPtr.Reset();
+	}
+
+	void FTickableNotification::Tick(float DeltaTime)
+	{
+		FScopeLock Lock(&CriticalSection);
+
+		if (bIsTicking && NotificationPtr.Pin().IsValid())
+		{
+			if (!WorkQueue.IsEmpty())
+			{
+				FWorkItem WorkItem;
+				WorkQueue.Dequeue(WorkItem);
+				WorkItem.Task(DisplayText);
+				FinalState = WorkItem.FinalState;
+				bIsTicking = !WorkItem.bIsFinalItem;
+			}
+
+			NotificationPtr.Pin()->SetText(DisplayText);
+		}
+		else
+		{
+			NotifyDestruction();
+		}
+	}
+
+	TStatId FTickableNotification::GetStatId() const
+	{
+		RETURN_QUICK_DECLARE_CYCLE_STAT(FTickableNotification, STATGROUP_Tickables);
+	}
+}
diff --git a/Engine/Plugins/Runtime/Steam/SteamAudio/Source/SteamAudioEditor/Private/TickableNotification.h b/Engine/Plugins/Runtime/Steam/SteamAudio/Source/SteamAudioEditor/Private/TickableNotification.h
new file mode 100644
--- /dev/null
+++ b/Engine/Plugins/Runtime/Steam/SteamAudio/Source/SteamAudioEditor/Private/TickableNotification.h
@@ -0,0 +1,49 @@
+//
+// Copyright (C) Valve Corporation. All rights reserved.
+//
+
+#pragma once
+
+#include "TickableEditorObject.h"
+#include "Widgets/Notifications/SNotificationList.h"
+#include "Containers/Queue.h"
+
+namespace SteamAudio
+{
+	struct FWorkItem
+	{
+		FWorkItem();
+		FWorkItem(const TFunction<void(FText&)>& Task, const SNotificationItem::ECompletionState InFinalState, const bool bIsFinalItem);
+
+		TFunction<void(FText&)> Task;
+		SNotificationItem::ECompletionState FinalState;
+		bool bIsFinalItem;
+	};
+
+	class FTickableNotification : public FTickableEditorObject
+	{
+	public:
+		FTickableNotification();
+
+		void CreateNotification();
+		void CreateNotificationWithCancel(const FSimpleDelegate& CancelDelegate);
+		void DestroyNotification(const SNotificationItem::ECompletionState InFinalState = SNotificationItem::CS_Success);
+		void SetDisplayText(const FText& InDisplayText);
+		void QueueWorkItem(const FWorkItem& WorkItem);
+
+	protected:
+		virtual void Tick(float DeltaTime) override;
+		virtual ETickableTickType GetTickableTickType() const override { return ETickableTickType::Always; }
+		virtual TStatId GetStatId() const override;
+
+	private:
+		void NotifyDestruction();
+
+		TWeakPtr<SNotificationItem> NotificationPtr;
+		FCriticalSection CriticalSection;
+		FText DisplayText;
+		bool bIsTicking;
+		TQueue<FWorkItem> WorkQueue;
+		SNotificationItem::ECompletionState FinalState;
+	};
+}
diff --git a/Engine/Plugins/Runtime/Steam/SteamAudio/Source/SteamAudioEditor/Public/ISteamAudioEditorModule.h b/Engine/Plugins/Runtime/Steam/SteamAudio/Source/SteamAudioEditor/Public/ISteamAudioEditorModule.h
new file mode 100644
--- /dev/null
+++ b/Engine/Plugins/Runtime/Steam/SteamAudio/Source/SteamAudioEditor/Public/ISteamAudioEditorModule.h
@@ -0,0 +1,33 @@
+//
+// Copyright (C) Valve Corporation. All rights reserved.
+//
+
+#pragma once
+
+#include "Modules/ModuleInterface.h"
+#include "Modules/ModuleManager.h"
+
+class ISteamAudioEditorModule : public IModuleInterface
+{
+public:
+	/**
+	 * Singleton-like access to this module's interface.  This is just for convenience!
+	 * Beware of calling this during the shutdown phase, though.  Your module might have been unloaded already.
+	 *
+	 * @return Returns singleton instance, loading the module on demand if needed
+	 */
+	static inline ISteamAudioEditorModule& Get()
+	{
+		return FModuleManager::LoadModuleChecked<ISteamAudioEditorModule>("SteamAudioEditorModule");
+	}
+
+	/**
+	 * Checks to see if this module is loaded and ready.  It is only valid to call Get() if IsAvailable() returns true.
+	 *
+	 * @return True if the module is loaded and ready to use
+	 */
+	static inline bool IsAvailable()
+	{
+		return FModuleManager::Get().IsModuleLoaded("SteamAudioEditorModule");
+	}
+};
diff --git a/Engine/Plugins/Runtime/Steam/SteamAudio/Source/SteamAudioEditor/SteamAudioEditor.Build.cs b/Engine/Plugins/Runtime/Steam/SteamAudio/Source/SteamAudioEditor/SteamAudioEditor.Build.cs
new file mode 100644
--- /dev/null
+++ b/Engine/Plugins/Runtime/Steam/SteamAudio/Source/SteamAudioEditor/SteamAudioEditor.Build.cs
@@ -0,0 +1,58 @@
+//
+// Copyright (C) Valve Corporation. All rights reserved.
+//
+
+namespace UnrealBuildTool.Rules
+{
+	public class SteamAudioEditor : ModuleRules
+	{
+		public SteamAudioEditor(ReadOnlyTargetRules Target) : base(Target)
+		{
+			PublicDependencyModuleNames.AddRange(
+				new string[] {
+					"Core",
+					"CoreUObject",
+					"Engine",
+					"InputCore",
+					"EditorFramework",
+					"UnrealEd",
+					"LevelEditor",
+					
+					"RenderCore",
+					"RHI",
+					"AudioEditor",
+					"AudioMixer",
+					"SteamAudio"
+				}
+			);
+
+			PrivateIncludePathModuleNames.AddRange(
+				new string[] {
+				    "AssetTools",
+				    "Landscape"
+			});
+
+			PrivateDependencyModuleNames.AddRange(
+				new string[] {
+					"Slate",
+					"SlateCore",
+					"UnrealEd",
+					"AudioEditor",
+					"LevelEditor",
+					"Landscape",
+					"Core",
+					"CoreUObject",
+					"Engine",
+					"InputCore",
+					"PropertyEditor",
+					"Projects",
+					
+					"SteamAudio"
+				 }
+			);
+
+			AddEngineThirdPartyPrivateStaticDependencies(Target, "libPhonon");
+			AddEngineThirdPartyPrivateStaticDependencies(Target, "DX11Audio");
+		}
+	}
+}
\ No newline at end of file
