[Backout] - CL25430596

diff --git a/Engine/Plugins/Runtime/Steam/SteamAudio/SteamAudio.uplugin b/Engine/Plugins/Runtime/Steam/SteamAudio/SteamAudio.uplugin
--- a/Engine/Plugins/Runtime/Steam/SteamAudio/SteamAudio.uplugin
+++ b/Engine/Plugins/Runtime/Steam/SteamAudio/SteamAudio.uplugin
@@ -3,7 +3,7 @@
 	"Version" : 1,
 	"VersionName" : "2.0-beta.17",
 	"FriendlyName" : "Steam Audio (Deprecated)",
-	"Description" : "This plugin is deprecated and will be removed in a future engine release. Please use the plugin from Valve's website.",
+	"Description" : "Physically-based sound rendering. This plugin is deprecated and will be removed in a future engine release. Please use the plugin from Valve's website.",
 	"Category" : "Audio",
 	"CreatedBy" : "Valve Corporation",
 	"CreatedByURL" : "https://valvesoftware.github.io/steam-audio",
@@ -24,6 +24,15 @@
 			[
 				"Win64", "Android"
 			]
+		},
+		{
+			"Name" : "SteamAudioEditor",
+			"Type" : "Editor",
+			"LoadingPhase" : "PostEngineInit",
+			"PlatformAllowList" :
+			[
+				"Win64"
+			]
 		}
 	]
 }
