Disable deprecated Steam Audio plugin on Windows Arm64

#rnx

diff --git a/Engine/Plugins/Runtime/Steam/SteamAudio/SteamAudio.uplugin b/Engine/Plugins/Runtime/Steam/SteamAudio/SteamAudio.uplugin
--- a/Engine/Plugins/Runtime/Steam/SteamAudio/SteamAudio.uplugin
+++ b/Engine/Plugins/Runtime/Steam/SteamAudio/SteamAudio.uplugin
@@ -23,6 +23,10 @@
 			"PlatformAllowList" :
 			[
 				"Win64", "Android"
+			],
+			"PlatformArchitectureDenyList":
+			[
+				"Win64:arm64"
 			]
 		}
 	]
